#UE4
- Start filling out the OSS session interface for Xbox Live, based on Lionhead's code. Basic flow is working: Create, Join, Find, Destroy etc all work through Xbox Live
- Added an IPv6 socket subsystem, also based on Lionhead's code (required for Xbox One)
- FURL constructor can now parse IPv6 hosts, IPv6 address must be enclosed in square brackets.
- Added new async item class, FOnlineAsyncItemGenericCallable, which makes is easy to use lambdas as async tasks.
- Temporarily set the ShooterGame defaults so that bIsLANMatch is true. This is to avoid communicating with Xbox Live during the first (non-multiplayer) cert pass.
- The ShooterGame server list would work over Xbox Live if the console menu exposed it. Leaving it off for first cert pass.

diff --git a/Engine/Source/Runtime/Engine/Private/URL.cpp b/Engine/Source/Runtime/Engine/Private/URL.cpp
--- a/Engine/Source/Runtime/Engine/Private/URL.cpp
+++ b/Engine/Source/Runtime/Engine/Private/URL.cpp
@@ -50,8 +50,7 @@ static bool ValidNetChar( const TCHAR* c )
 
 	// / Is now allowed because absolute paths are required in various places
 
-	if( FCString::Strchr( c, ':' ) ||		// : is for protocol and user/password info
-		FCString::Strchr( c, '?' ) || FCString::Strchr( c, '#' ) )		// ? and # delimit fragments
+	if( FCString::Strchr( c, '?' ) || FCString::Strchr( c, '#' ) )		// ? and # delimit fragments
 	{		
 		return false;
 	}
@@ -236,10 +235,16 @@ FURL::FURL( FURL* Base, const TCHAR* TextURL, ETravelType Type )
 		}
 		else
 		{
-			// Parse protocol.
+			// Determine location of the first opening square bracket.
+			// Square brackets enclose an IPv6 address.
+			const TCHAR* SquareBracket = FCString::Strchr(URL, '[');
+
+			// Parse protocol. Don't consider colons that occur after the opening square
+			// brace, because they are valid characters in an IPv6 address.
 			if
 			(	(FCString::Strchr(URL,':')!=NULL)
 			&&	(FCString::Strchr(URL,':')>URL+1)
+			&&  (FCString::Strchr(URL,':')<SquareBracket)
 			&&	(FCString::Strchr(URL,'.')==NULL || FCString::Strchr(URL,':')<FCString::Strchr(URL,'.')) )
 			{
 				TCHAR* ss = URL;
@@ -291,6 +296,27 @@ FURL::FURL( FURL* Base, const TCHAR* TextURL, ETravelType Type )
 				}
 				FarHost = 1;
 			}
+
+			// Parse IPv6 host.
+			if(URL && SquareBracket)
+			{
+				const TCHAR* ss = URL;
+				URL = FCString::Strchr(URL, ']');
+				if (URL)
+				{
+					*URL = 0;
+					Host = SquareBracket + 1;
+				}
+				if( FCString::Stricmp(*Protocol,*UrlConfig.DefaultProtocol)==0 )
+				{
+					Map = UGameMapsSettings::GetGameDefaultMap();
+				}
+				else
+				{
+					Map = TEXT("");
+				}
+				FarHost = 1;
+			}
 		}
 	}
 
diff --git a/Engine/Source/Runtime/Online/OnlineSubsystem/Private/OnlineSessionSettings.cpp b/Engine/Source/Runtime/Online/OnlineSubsystem/Private/OnlineSessionSettings.cpp
--- a/Engine/Source/Runtime/Online/OnlineSubsystem/Private/OnlineSessionSettings.cpp
+++ b/Engine/Source/Runtime/Online/OnlineSubsystem/Private/OnlineSessionSettings.cpp
@@ -44,7 +44,7 @@ void DumpSession(const FOnlineSession* Session)
 	{
 		UE_LOG(LogOnline, Verbose, TEXT("dumping Session: "));
 		UE_LOG(LogOnline, Verbose, TEXT("	OwningPlayerName: %s"), *Session->OwningUserName);	
-		UE_LOG(LogOnline, Verbose, TEXT("	OwningPlayerId: %s"), *Session->OwningUserId->ToDebugString());
+		UE_LOG(LogOnline, Verbose, TEXT("	OwningPlayerId: %s"), Session->OwningUserId.IsValid() ? *Session->OwningUserId->ToDebugString() : TEXT("") );
 		UE_LOG(LogOnline, Verbose, TEXT("	NumOpenPrivateConnections: %d"), Session->NumOpenPrivateConnections);	
 		UE_LOG(LogOnline, Verbose, TEXT("	NumOpenPublicConnections: %d"), Session->NumOpenPublicConnections);	
 		UE_LOG(LogOnline, Verbose, TEXT("	SessionInfo: %s"), Session->SessionInfo.IsValid() ? *Session->SessionInfo->ToDebugString() : TEXT("NULL"));
diff --git a/Engine/Source/Runtime/Online/OnlineSubsystem/Public/OnlineAsyncTaskManager.h b/Engine/Source/Runtime/Online/OnlineSubsystem/Public/OnlineAsyncTaskManager.h
--- a/Engine/Source/Runtime/Online/OnlineSubsystem/Public/OnlineAsyncTaskManager.h
+++ b/Engine/Source/Runtime/Online/OnlineSubsystem/Public/OnlineAsyncTaskManager.h
@@ -59,6 +59,36 @@ public:
 	}
 };
 
+/**
+ * An async item that can execute any callable type with no parameters.
+ * For example, l lambda, or an object with an operator().
+ * Useful for calling simple functions that need to run on the game thread,
+ * but are invoked from an online service thread.
+ */
+template<class CallableType>
+class FOnlineAsyncItemGenericCallable : public FOnlineAsyncItem
+{
+public:
+	/**
+	 * Constructor.
+	 *
+	 * @param InCallable any object that can be called with no parameters, usually a lambda
+	 */
+	explicit FOnlineAsyncItemGenericCallable(const CallableType& InCallable)
+		: CallableObject(InCallable) {}
+
+	virtual void Finalize() OVERRIDE
+	{
+		CallableObject();
+	}
+
+	virtual FString ToString() const OVERRIDE { return FString("FOnlineAsyncItemGenericCallable"); }
+
+private:
+	/** Stored copy of the object to invoke on the game thread. */
+	CallableType CallableObject;
+};
+
 /**
  * An event triggered by the online subsystem to be routed to the game thread for processing
  * Originates on the online thread
@@ -271,6 +301,19 @@ public:
 	 */
 	void AddToOutQueue(FOnlineAsyncItem* CompletedItem);
 
+	/**
+	 * Add a new item to the out queue that will call InCallable on the game thread.
+	 * Very useful when passing in lambdas as parameters, since this function will
+	 * automatically deduce the template parameter type for FOnlineAsyncItemGenericCallable.
+	 *
+	 * @param InCallable the callable object to execute on the game thread.
+	 */
+	template<class CallableType>
+	void AddGenericToOutQueue(const CallableType& InCallable)
+	{
+		AddToOutQueue(new FOnlineAsyncItemGenericCallable<CallableType>(InCallable));
+	}
+
 	/**
 	 * Add a new online async task that is safe to run in parallel
 	 * @param NewTask - some request of the online services
diff --git a/Engine/Source/Runtime/Online/OnlineSubsystem/Public/OnlineSessionSettings.h b/Engine/Source/Runtime/Online/OnlineSubsystem/Public/OnlineSessionSettings.h
--- a/Engine/Source/Runtime/Online/OnlineSubsystem/Public/OnlineSessionSettings.h
+++ b/Engine/Source/Runtime/Online/OnlineSubsystem/Public/OnlineSessionSettings.h
@@ -5,6 +5,7 @@
 #include "OnlineSubsystemTypes.h"
 #include "OnlineKeyValuePair.h"
 #include "OnlineSubsystemPackage.h"
+#include "OnlineSubsystemSessionSettings.h"
 
 /** Setting describing the name of the current map (value is FString) */
 #define SETTING_MAPNAME FName(TEXT("MAPNAME"))
@@ -432,6 +433,10 @@ public:
 #define SEARCH_PRESENCE FName(TEXT("PRESENCESEARCH"))
 /** Search for a match with min player availability (value is int) */
 #define SEARCH_MINSLOTSAVAILABLE FName(TEXT("MINSLOTSAVAILABLE"))
+/** User ID to search for session of */
+#define SEARCH_USER FName(TEXT("SEARCHUSER"))
+/** Keywords to match in session search */
+#define SEARCH_KEYWORDS FName(TEXT("SEARCHKEYWORDS"))
 
 /** 
  * Encapsulation of a search for sessions request.
diff --git a/Engine/Source/Runtime/Online/OnlineSubsystem/Public/OnlineSubsystemSessionSettings.h b/Engine/Source/Runtime/Online/OnlineSubsystem/Public/OnlineSubsystemSessionSettings.h
new file mode 100644
--- /dev/null
+++ b/Engine/Source/Runtime/Online/OnlineSubsystem/Public/OnlineSubsystemSessionSettings.h
@@ -0,0 +1,35 @@
+// Copyright 1998-2013 Epic Games, Inc. All Rights Reserved.
+
+#pragma once
+#include "OnlineSubsystemPackage.h"
+
+// These are Xbox One specific setting to be used in FOnlineSessionSettings as keys; technically this can be put in OnlineSubsystemSettings.h
+
+/** Custom settings to be associated with Session (value is FString)*/
+#define SETTING_CUSTOM FName(TEXT("CUSTOM"))
+/** Session is marked as using party enabled session flag (value is bool)*/
+#define SETTING_PARTY_ENABLED_SESSION FName(TEXT("PARTYENABLEDSESSION"))
+/** Hopper name to search within (value is FString) */
+#define SETTING_MATCHING_HOPPER FName(TEXT("MATCHHOPPER"))
+/** Timeout for search (value is float seconds) */
+#define SETTING_MATCHING_TIMEOUT FName(TEXT("MATCHTIMEOUT"))
+/** Match attributes (value is FString) */
+#define SETTING_MATCHING_ATTRIBUTES	FName(TEXT("MATCHATTRIBUTES"))
+/** Match attributes (value is FString) */
+#define SETTING_MATCH_MEMBERS_JSON FName(TEXT("MATCHMEMBERS"))
+/** Session member constant custom json (value is FString) - This is to be used with FString::Printf, populated with user xuid */
+#define SETTING_SESSION_MEMBER_CONSTANT_CUSTOM_JSON_XUID FName(TEXT("SESSIONMEMBERCONSTANTCUSTOMJSON%s"))
+/** Set self as host (value is int32)*/
+#define SETTING_MAX_RESULT FName(TEXT("MAXRESULT"))
+/** Set self as host (value is int32)*/
+#define SETTING_CONTRACT_VERSION_FILTER FName(TEXT("CONTRACTVERSIONFILTER"))
+/** Set self as host (value is int32)*/
+#define SETTING_FIND_PRIVATE_SESSIONS FName(TEXT("FINDPRIVATESESSIONS"))
+/** Set self as host (value is int32)*/
+#define SETTING_FIND_RESERVED_SESSIONS FName(TEXT("FINDRESERVEDSESSIONS"))
+/** Set self as host (value is int32)*/
+#define SETTING_FIND_INACTIVE_SESSIONS FName(TEXT("FINDINACTIVESESSIONS"))
+/** Set self as host (value is int32)*/
+#define SETTING_MULTIPLAYER_VISIBILITY FName(TEXT("MULTIPLAYERVISIBILITY"))
+/** Session template name, must match a session template name from the service config. */
+#define SETTING_SESSION_TEMPLATE_NAME FName(TEXT("SESSIONTEMPLATENAME"))
\ No newline at end of file
diff --git a/Engine/Source/Runtime/Online/OnlineSubsystemUtils/Private/IpNetDriver.cpp b/Engine/Source/Runtime/Online/OnlineSubsystemUtils/Private/IpNetDriver.cpp
--- a/Engine/Source/Runtime/Online/OnlineSubsystemUtils/Private/IpNetDriver.cpp
+++ b/Engine/Source/Runtime/Online/OnlineSubsystemUtils/Private/IpNetDriver.cpp
@@ -97,14 +97,20 @@ bool UIpNetDriver::InitBase( bool bInitAsClient, FNetworkNotify* InNotify, const
 
 	// Bind socket to our port.
 	LocalAddr = SocketSubsystem->GetLocalBindAddr(*GLog);
-	LocalAddr->SetPort(0);
-	if(!bInitAsClient)
-	{
-		// Init as a server.
-		LocalAddr->SetPort(URL.Port);
-	}
+	
+	// Try the specific port first, even for clients. Required for some platforms.
+	LocalAddr->SetPort(URL.Port);
+	
 	int32 AttemptPort = LocalAddr->GetPort();
 	int32 BoundPort = SocketSubsystem->BindNextPort( Socket, *LocalAddr, MaxPortCountToTry + 1, 1 );
+	if (BoundPort == 0 && bInitAsClient)
+	{
+		// If initializing as a client, and the specific port failed, retry with port 0 (any port)
+		LocalAddr->SetPort(0);
+		AttemptPort = LocalAddr->GetPort();
+		int32 BoundPort = SocketSubsystem->BindNextPort( Socket, *LocalAddr, 1, 1 );
+	}
+
 	if (BoundPort == 0)
 	{
 		Error = FString::Printf( TEXT("%s: binding to port %i failed (%i)"), SocketSubsystem->GetSocketAPIName(), AttemptPort,
diff --git a/Engine/Source/Runtime/Sockets/Private/BSDIPv6Sockets/IPAddressBSDIPv6.h b/Engine/Source/Runtime/Sockets/Private/BSDIPv6Sockets/IPAddressBSDIPv6.h
new file mode 100644
--- /dev/null
+++ b/Engine/Source/Runtime/Sockets/Private/BSDIPv6Sockets/IPAddressBSDIPv6.h
@@ -0,0 +1,269 @@
+// Copyright 1998-2014 Epic Games, Inc. All Rights Reserved.
+
+#pragma once
+
+#if PLATFORM_HAS_BSD_SOCKETS
+
+#include "Core.h"
+#include "IPAddress.h"
+
+/**
+ * Represents an internet ip address, using the relatively standard sockaddr_in6 structure. All data is in network byte order
+ */
+class FInternetAddrBSDIPv6 : public FInternetAddr
+{
+	/** The internet ip address structure */
+	sockaddr_in6 Addr;
+
+	/** Horrible hack to catch hard coded multicasting on IPv4 **/
+	static const uint32 IPv4MulitcastAddr = ((230 << 24) | (0 << 16) | (0 << 8) | (1 << 0));
+
+public:
+	/**
+	 * Constructor. Sets address to default state
+	 */
+	FInternetAddrBSDIPv6()
+	{
+		FMemory::Memzero(&Addr,sizeof(Addr));
+		Addr.sin6_family = AF_INET6;
+	}
+
+	/**
+	 * Sets the ip address from a host byte order uint32, convert the IPv4 address supplied to an IPv6 address
+	 *
+	 * @param InAddr the new address to use (must convert to network byte order)
+	 */
+	virtual void SetIp(uint32 InAddr) OVERRIDE
+	{
+		if(InAddr == 0)
+		{
+			FMemory::Memzero(&Addr.sin6_addr,sizeof(Addr.sin6_addr));
+		}
+		else if(InAddr == IPv4MulitcastAddr)
+		{
+			// if it's the hardcoded IPv4 multicasting address then translate into an IPv6 multicast address
+			bool isValid;
+			SetIp(L"ff02::2", isValid);
+			check(isValid);
+		}
+		else
+		{
+			in_addr Addr;
+			Addr.S_un.S_addr= htonl(InAddr);
+
+			SetIp(Addr);
+		}
+	}
+
+	/**
+	 * Sets the ip address from a string ("[aaaa:bbbb:cccc:dddd:eeee:ffff:gggg:hhhh]:port" or "a.b.c.d:port")
+	 *
+	 * @param InAddr the string containing the new ip address to use
+	 * @param bIsValid will be set to true if InAddr was a valid IPv6 or IPv4 address, false if not.
+	 */
+	virtual void SetIp(const TCHAR* InAddr, bool& bIsValid) OVERRIDE
+	{
+		// check for valid IPv6 address
+		if(InetPtonW(AF_INET6, InAddr, &Addr.sin6_addr))
+		{
+			bIsValid = true;
+			return;
+		}
+
+		// IPv6 URLs are surrounded by square brackets, check for that.
+		if(InAddr && *InAddr == '[')
+		{
+			// Make a copy, skipping the opening brace.
+			FString CopiedAddr(InAddr + 1);
+
+			// Remove the closing brace if it exists
+			CopiedAddr.RemoveFromEnd(TEXT("]"));
+
+			if(InetPtonW(AF_INET6, CopiedAddr.GetCharArray().GetTypedData(), &Addr.sin6_addr))
+			{
+				bIsValid = true;
+				return;
+			}
+		}
+
+		// Check if it's a valid IPv4 address, and if it is convert
+		in_addr  IPv4Addr;
+		if(InetPtonW(AF_INET, InAddr, &IPv4Addr))
+		{
+			bIsValid = true;
+			SetIp(IPv4Addr);
+		}
+		else
+		{
+			bIsValid = false;
+		}
+	}
+
+	/**
+	 * Sets the ip address using a network byte order ipv4 address
+	 *
+	 * @param IPv4Addr the new ip address to use
+	 */	
+	void SetIp(const in_addr& IPv4Addr)
+	{
+		FMemory::Memzero(&Addr.sin6_addr,sizeof(Addr.sin6_addr));
+
+		// special mapping of ipv4 to ipv6 using a hybrid stack, won't work on a pure ipv6 implementation
+		Addr.sin6_addr.u.Word[5] = 0xffff;
+		Addr.sin6_addr.u.Byte[12] = IPv4Addr.S_un.S_un_b.s_b1;
+		Addr.sin6_addr.u.Byte[13] = IPv4Addr.S_un.S_un_b.s_b2;
+		Addr.sin6_addr.u.Byte[14] = IPv4Addr.S_un.S_un_b.s_b3;
+		Addr.sin6_addr.u.Byte[15] = IPv4Addr.S_un.S_un_b.s_b4;
+
+		UE_LOG(LogSockets, Log, TEXT("Using IPv4 address: %d.%d.%d.%d  on an ipv6 socket"), IPv4Addr.S_un.S_un_b.s_b1, IPv4Addr.S_un.S_un_b.s_b2, IPv4Addr.S_un.S_un_b.s_b3, IPv4Addr.S_un.S_un_b.s_b4);
+	}
+
+	/**
+	 * Sets the ip address using a network byte order ipv6 address
+	 *
+	 * @param IpAddr the new ip address to use
+	 */
+	void SetIp(const in6_addr& IpAddr)
+ 	{
+ 		Addr.sin6_addr = IpAddr;
+ 	}
+
+	/**
+	 * Sets the ip address using a generic sockaddr_storage
+	 *
+	 * @param IpAddr the new ip address to use
+	 */
+	void SetIp(const sockaddr_storage& IpAddr)
+	{
+		if (IpAddr.ss_family == AF_INET)
+		{
+			const sockaddr_in* Addr = (const sockaddr_in*)&IpAddr;
+			SetIp(Addr->sin_addr);
+		}
+		else if (IpAddr.ss_family == AF_INET6)
+		{
+			const sockaddr_in6* Addr = (const sockaddr_in6*)&IpAddr;
+			SetIp(Addr->sin6_addr);
+		}
+	}
+
+	/**
+	 * Copies the network byte order ip address to a host byte order dword, doesn't exist with IPv6
+	 *
+	 * @param OutAddr the out param receiving the ip address
+	 */
+	virtual void GetIp(uint32& OutAddr) const OVERRIDE
+	{
+		// grab the last 32-bits of the IPv6 address as this will correspond to the IPv4 address
+		// in a dual stack system.
+		// This function doesn't really make sense in IPv6, but too much other code relies on it
+		// existing to not have this here.
+		OutAddr = (Addr.sin6_addr.u.Word[6] << 16) & (Addr.sin6_addr.u.Word[7]);
+	}
+#if 0
+	/**
+	 * Copies the network byte order ip address 
+	 *
+	 * @param OutAddr the out param receiving the ip address
+	 */
+	void GetIp(in6_addr& OutAddr) const
+ 	{
+ 		OutAddr = Addr.sin6_addr;
+ 	}
+#endif
+	/**
+	 * Sets the port number from a host byte order int
+	 *
+	 * @param InPort the new port to use (must convert to network byte order)
+	 */
+	virtual void SetPort(int32 InPort) OVERRIDE
+	{
+		Addr.sin6_port = htons(InPort);
+	}
+
+	/**
+	 * Copies the port number from this address and places it into a host byte order int
+	 *
+	 * @param OutPort the host byte order int that receives the port
+	 */
+	virtual void GetPort(int32& OutPort) const OVERRIDE
+	{
+		OutPort = ntohs(Addr.sin6_port);
+	}
+
+	/** Returns the port number from this address in host byte order */
+	virtual int32 GetPort() const OVERRIDE
+	{
+		return ntohs(Addr.sin6_port);
+	}
+
+	/** Sets the address to be any address */
+	virtual void SetAnyAddress() OVERRIDE
+	{
+		SetIp(in6addr_any);
+		SetPort(0);
+	}
+
+	/** Sets the address to broadcast */
+	virtual void SetBroadcastAddress() OVERRIDE
+	{
+		// broadcast means something different in IPv6, but this is a rough equivalent
+		SetIp(in6addr_allnodesonlink);
+		SetPort(0);
+	}
+
+	/**
+	 * Converts this internet ip address to string form. String will be enclosed in square braces.
+	 *
+	 * @param bAppendPort whether to append the port information or not
+	 */
+	virtual FString ToString(bool bAppendPort) const OVERRIDE
+	{
+		char IPStr[INET6_ADDRSTRLEN];
+
+		inet_ntop(AF_INET6, (void*)&Addr.sin6_addr, IPStr, INET6_ADDRSTRLEN);
+
+		FString Result("[");
+		Result += IPStr;
+		Result += "]";
+
+		return Result;
+	}
+
+	/**
+	 * Compares two internet ip addresses for equality
+	 *
+	 * @param Other the address to compare against
+	 */
+	virtual bool operator==(const FInternetAddr& Other) const OVERRIDE
+	{
+		FInternetAddrBSDIPv6& OtherBSD = (FInternetAddrBSDIPv6&)Other;
+		return memcmp(&Addr.sin6_addr,&OtherBSD.Addr.sin6_addr,sizeof(in6_addr)) == 0 &&
+			Addr.sin6_port == OtherBSD.Addr.sin6_port &&
+			Addr.sin6_family == OtherBSD.Addr.sin6_family;
+	}
+
+	/**
+	 * Is this a well formed internet address, the only criteria being non-zero
+	 *
+	 * @return true if a valid IP, false otherwise
+	 */
+	virtual bool IsValid() const OVERRIDE
+	{
+		FInternetAddrBSDIPv6 Temp;
+
+		return memcmp(&Addr.sin6_addr, &Temp.Addr.sin6_addr, sizeof(in6_addr)) != 0;
+	}
+
+ 	operator sockaddr*(void)
+ 	{
+ 		return (sockaddr*)&Addr;
+ 	}
+
+ 	operator const sockaddr*(void) const
+ 	{
+ 		return (const sockaddr*)&Addr;
+ 	}
+};
+
+#endif	//PLATFORM_HAS_BSD_SOCKETS
diff --git a/Engine/Source/Runtime/Sockets/Private/BSDIPv6Sockets/SocketSubsystemBSDIPv6.cpp b/Engine/Source/Runtime/Sockets/Private/BSDIPv6Sockets/SocketSubsystemBSDIPv6.cpp
new file mode 100644
--- /dev/null
+++ b/Engine/Source/Runtime/Sockets/Private/BSDIPv6Sockets/SocketSubsystemBSDIPv6.cpp
@@ -0,0 +1,197 @@
+// Copyright 1998-2014 Epic Games, Inc. All Rights Reserved.
+
+#include "SocketsPrivatePCH.h"
+
+#if PLATFORM_HAS_BSD_SOCKETS
+
+#include "SocketSubsystemBSDIPv6.h"
+#include "SocketsBSDIPv6.h"
+
+
+class FSocketBSDIPv6* FSocketSubsystemBSDIPv6::InternalBSDSocketFactory(SOCKET Socket, ESocketType SocketType, const FString& SocketDescription)
+{
+	// return a new socket object
+	return new FSocketBSDIPv6(Socket, SocketType, SocketDescription, this);
+}
+
+FSocket* FSocketSubsystemBSDIPv6::CreateSocket(const FName& SocketType, const FString& SocketDescription, bool bForceUDP)
+{
+	SOCKET Socket = INVALID_SOCKET;
+	FSocket* NewSocket = NULL;
+	switch (SocketType.GetIndex())
+	{
+	case NAME_DGram:
+		// Creates a data gram (UDP) socket
+		Socket = socket(AF_INET6, SOCK_DGRAM, IPPROTO_UDP);
+		NewSocket = (Socket != INVALID_SOCKET) ? InternalBSDSocketFactory(Socket, SOCKTYPE_Datagram, SocketDescription) : NULL;
+		break;
+	case NAME_Stream:
+		// Creates a stream (TCP) socket
+		Socket = socket(AF_INET6, SOCK_STREAM, IPPROTO_TCP);
+		NewSocket = (Socket != INVALID_SOCKET) ? InternalBSDSocketFactory(Socket, SOCKTYPE_Streaming, SocketDescription) : NULL;
+		break;
+	}
+
+	if (!NewSocket)
+	{
+		UE_LOG(LogSockets, Warning, TEXT("Failed to create IPv6 socket %s [%s]"), *SocketType.ToString(), *SocketDescription);
+	}
+
+	return NewSocket;
+}
+
+void FSocketSubsystemBSDIPv6::DestroySocket(FSocket* Socket)
+{
+	delete Socket;
+}
+
+ESocketErrors FSocketSubsystemBSDIPv6::GetHostByName(const ANSICHAR* HostName, FInternetAddr& OutAddr)
+{
+#if PLATFORM_HAS_BSD_SOCKET_FEATURE_GETHOSTNAME
+	ESocketErrors ErrorCode = SE_NO_ERROR;
+	// gethostbyname() touches a static object so lock for thread safety
+	FScopeLock ScopeLock(&HostByNameSynch);
+	hostent* HostEnt = gethostbyname(HostName);
+	if (HostEnt != NULL)
+	{
+		// Make sure it's a valid type
+		if (HostEnt->h_addrtype == PF_INET)
+		{
+			// Copy the data before letting go of the lock. This is safe only
+			// for the copy locally. If another thread is reading this while
+			// we are copying they will get munged data. This relies on the
+			// consumer of this class to call the resolved() accessor before
+			// attempting to read this data
+			((FInternetAddrBSDIPv6&)OutAddr).SetIp(*(in_addr6*)(*HostEnt->h_addr_list));
+		}
+		else
+		{
+			ErrorCode = SE_HOST_NOT_FOUND;
+		}
+	}
+	else
+	{
+		ErrorCode = GetLastErrorCode();
+	}
+	return ErrorCode;
+#else
+	UE_LOG(LogSockets, Error, TEXT("Platform has no gethostbyname(), but did not override FSocketSubsystem::GetHostByName()"));
+	return SE_NO_RECOVERY;
+#endif
+}
+
+bool FSocketSubsystemBSDIPv6::GetHostName(FString& HostName)
+{
+#if PLATFORM_HAS_BSD_SOCKET_FEATURE_GETHOSTNAME
+	ANSICHAR Buffer[256];
+	bool bRead = gethostname(Buffer,256) == 0;
+	if (bRead == true)
+	{
+		HostName = Buffer;
+	}
+	return bRead;
+#else
+	UE_LOG(LogSockets, Error, TEXT("Platform has no gethostname(), but did not override FSocketSubsystem::GetHostName()"));
+	return false;
+#endif
+}
+
+const TCHAR* FSocketSubsystemBSDIPv6::GetSocketAPIName() const
+{
+	return TEXT("BSD IPv6");
+}
+
+TSharedRef<FInternetAddr> FSocketSubsystemBSDIPv6::CreateInternetAddr(uint32 Address, uint32 Port)
+{
+	TSharedRef<FInternetAddr> Result = MakeShareable(new FInternetAddrBSDIPv6);
+	Result->SetIp(Address);
+	Result->SetPort(Port);
+	return Result;
+}
+
+ESocketErrors FSocketSubsystemBSDIPv6::GetLastErrorCode()
+{
+	return TranslateErrorCode(errno);
+}
+
+
+ESocketErrors FSocketSubsystemBSDIPv6::TranslateErrorCode(int32 Code)
+{
+	// @todo sockets: Windows for some reason doesn't seem to have all of the standard error messages,
+	// but it overrides this function anyway - however, this 
+#if !PLATFORM_HAS_BSD_SOCKET_FEATURE_WINSOCKETS
+
+	// handle the generic -1 error
+	if (Code == SOCKET_ERROR)
+	{
+		return GetLastErrorCode();
+	}
+
+	switch (Code)
+	{
+	case 0: return SE_NO_ERROR;
+	case EINTR: return SE_EINTR;
+	case EBADF: return SE_EBADF;
+	case EACCES: return SE_EACCES;
+	case EFAULT: return SE_EFAULT;
+	case EINVAL: return SE_EINVAL;
+	case EMFILE: return SE_EMFILE;
+	case EWOULDBLOCK: return SE_EWOULDBLOCK;
+	case EINPROGRESS: return SE_EINPROGRESS;
+	case EALREADY: return SE_EALREADY;
+	case ENOTSOCK: return SE_ENOTSOCK;
+	case EDESTADDRREQ: return SE_EDESTADDRREQ;
+	case EMSGSIZE: return SE_EMSGSIZE;
+	case EPROTOTYPE: return SE_EPROTOTYPE;
+	case ENOPROTOOPT: return SE_ENOPROTOOPT;
+	case EPROTONOSUPPORT: return SE_EPROTONOSUPPORT;
+	case ESOCKTNOSUPPORT: return SE_ESOCKTNOSUPPORT;
+	case EOPNOTSUPP: return SE_EOPNOTSUPP;
+	case EPFNOSUPPORT: return SE_EPFNOSUPPORT;
+	case EAFNOSUPPORT: return SE_EAFNOSUPPORT;
+	case EADDRINUSE: return SE_EADDRINUSE;
+	case EADDRNOTAVAIL: return SE_EADDRNOTAVAIL;
+	case ENETDOWN: return SE_ENETDOWN;
+	case ENETUNREACH: return SE_ENETUNREACH;
+	case ENETRESET: return SE_ENETRESET;
+	case ECONNABORTED: return SE_ECONNABORTED;
+	case ECONNRESET: return SE_ECONNRESET;
+	case ENOBUFS: return SE_ENOBUFS;
+	case EISCONN: return SE_EISCONN;
+	case ENOTCONN: return SE_ENOTCONN;
+	case ESHUTDOWN: return SE_ESHUTDOWN;
+	case ETOOMANYREFS: return SE_ETOOMANYREFS;
+	case ETIMEDOUT: return SE_ETIMEDOUT;
+	case ECONNREFUSED: return SE_ECONNREFUSED;
+	case ELOOP: return SE_ELOOP;
+	case ENAMETOOLONG: return SE_ENAMETOOLONG;
+	case EHOSTDOWN: return SE_EHOSTDOWN;
+	case EHOSTUNREACH: return SE_EHOSTUNREACH;
+	case ENOTEMPTY: return SE_ENOTEMPTY;
+	case EUSERS: return SE_EUSERS;
+	case EDQUOT: return SE_EDQUOT;
+	case ESTALE: return SE_ESTALE;
+	case EREMOTE: return SE_EREMOTE;
+#if !PLATFORM_HAS_NO_EPROCLIM
+	case EPROCLIM: return SE_EPROCLIM;
+#endif
+// 	case EDISCON: return SE_EDISCON;
+// 	case SYSNOTREADY: return SE_SYSNOTREADY;
+// 	case VERNOTSUPPORTED: return SE_VERNOTSUPPORTED;
+// 	case NOTINITIALISED: return SE_NOTINITIALISED;
+#if PLATFORM_HAS_BSD_SOCKET_FEATURE_GETHOSTNAME
+	case HOST_NOT_FOUND: return SE_HOST_NOT_FOUND;
+	case TRY_AGAIN: return SE_TRY_AGAIN;
+	case NO_RECOVERY: return SE_NO_RECOVERY;
+#endif
+//	case NO_DATA: return SE_NO_DATA;
+		// case : return SE_UDP_ERR_PORT_UNREACH; //@TODO Find it's replacement
+	}
+#endif
+
+	UE_LOG(LogSockets, Warning, TEXT("Unhandled socket error!"));
+	check(0);
+	return SE_NO_ERROR;
+}
+
+#endif	//PLATFORM_HAS_BSD_SOCKETS
diff --git a/Engine/Source/Runtime/Sockets/Private/BSDIPv6Sockets/SocketSubsystemBSDIPv6.h b/Engine/Source/Runtime/Sockets/Private/BSDIPv6Sockets/SocketSubsystemBSDIPv6.h
new file mode 100644
--- /dev/null
+++ b/Engine/Source/Runtime/Sockets/Private/BSDIPv6Sockets/SocketSubsystemBSDIPv6.h
@@ -0,0 +1,77 @@
+// Copyright 1998-2014 Epic Games, Inc. All Rights Reserved.
+
+#pragma once
+
+#include "SocketSubsystem.h"
+#include "SocketSubsystemPackage.h"
+
+#if PLATFORM_HAS_BSD_SOCKETS
+
+#include "../BSDSockets/SocketSubsystemBSDPrivate.h"
+#include "IPAddressBSDIPv6.h"
+
+
+/**
+ * Standard BSD specific IPv6 socket subsystem implementation
+ */
+class FSocketSubsystemBSDIPv6
+	: public ISocketSubsystem
+{
+public:
+
+	// Begin ISocketSubsystem interface
+
+	virtual TSharedRef<FInternetAddr> CreateInternetAddr( uint32 Address = 0, uint32 Port = 0 ) OVERRIDE;
+
+	virtual class FSocket* CreateSocket( const FName& SocketType, const FString& SocketDescription, bool bForceUDP = false ) OVERRIDE;
+
+	virtual void DestroySocket( class FSocket* Socket ) OVERRIDE;
+
+	virtual ESocketErrors GetHostByName( const ANSICHAR* HostName, FInternetAddr& OutAddr ) OVERRIDE;
+
+	virtual bool GetHostName( FString& HostName ) OVERRIDE;
+
+	virtual ESocketErrors GetLastErrorCode( ) OVERRIDE;
+
+	virtual bool GetLocalAdapterAddresses( TArray<TSharedPtr<FInternetAddr> >& OutAdresses ) OVERRIDE
+	{
+		bool bCanBindAll;
+
+		OutAdresses.Add(GetLocalHostAddr(*GLog, bCanBindAll));
+
+		return true;
+	}
+
+	virtual const TCHAR* GetSocketAPIName( ) const OVERRIDE;
+
+	virtual bool RequiresChatDataBeSeparate( ) OVERRIDE
+	{
+		return false;
+	}
+
+	virtual bool RequiresEncryptedPackets( ) OVERRIDE
+	{
+		return false;
+	}
+
+	virtual ESocketErrors TranslateErrorCode( int32 Code ) OVERRIDE;
+
+	// End ISocketSubsystem interface
+
+
+protected:
+
+	/**
+	 * Allows a subsystem subclass to create a FSocketBSD sub class.
+	 */
+	virtual class FSocketBSDIPv6* InternalBSDSocketFactory( SOCKET Socket, ESocketType SocketType, const FString& SocketDescription );
+
+	// allow BSD sockets to use this when creating new sockets from accept() etc
+	friend FSocketBSDIPv6;
+
+private:
+
+	// Used to prevent multiple threads accessing the shared data.
+	FCriticalSection HostByNameSynch;
+};
+#endif
\ No newline at end of file
diff --git a/Engine/Source/Runtime/Sockets/Private/BSDIPv6Sockets/SocketsBSDIPv6.cpp b/Engine/Source/Runtime/Sockets/Private/BSDIPv6Sockets/SocketsBSDIPv6.cpp
new file mode 100644
--- /dev/null
+++ b/Engine/Source/Runtime/Sockets/Private/BSDIPv6Sockets/SocketsBSDIPv6.cpp
@@ -0,0 +1,407 @@
+// Copyright 1998-2014 Epic Games, Inc. All Rights Reserved.
+
+#include "SocketsPrivatePCH.h"
+
+#if PLATFORM_HAS_BSD_SOCKETS
+
+#include "SocketSubsystemBSDIPv6.h"
+#include "SocketsBSDIPv6.h"
+//#include "Net/NetworkProfiler.h"
+
+bool FSocketBSDIPv6::Close(void)
+{
+	return closesocket(Socket) == 0;
+}
+
+
+bool FSocketBSDIPv6::Bind(const FInternetAddr& Addr)
+{
+	return bind(Socket, (sockaddr*)(FInternetAddrBSDIPv6&)Addr, sizeof(sockaddr_in6)) == 0;
+}
+
+
+bool FSocketBSDIPv6::Connect(const FInternetAddr& Addr)
+{
+	int32 Return = connect(Socket, (sockaddr*)(FInternetAddrBSDIPv6&)Addr, sizeof(sockaddr_in6));
+	
+	check(SocketSubsystem);
+	ESocketErrors Error = SocketSubsystem->TranslateErrorCode(Return);
+
+	// "would block" is not an error
+	return ((Error == SE_NO_ERROR) || (Error == SE_EWOULDBLOCK));
+}
+
+
+bool FSocketBSDIPv6::Listen(int32 MaxBacklog)
+{
+	return listen(Socket, MaxBacklog) == 0;
+}
+
+EIPv6SocketInternalState::Return FSocketBSDIPv6::HasState(EIPv6SocketInternalState::Param State, FTimespan WaitTime)
+{
+#if PLATFORM_HAS_BSD_SOCKET_FEATURE_SELECT
+	// convert WaitTime to a timeval
+	timeval Time;
+	Time.tv_sec = (int32)WaitTime.GetTotalSeconds();
+	Time.tv_usec = WaitTime.GetMilliseconds() * 1000;
+
+	fd_set SocketSet;
+
+	// Set up the socket sets we are interested in (just this one)
+	FD_ZERO(&SocketSet);
+	FD_SET(Socket, &SocketSet);
+
+	// Check the status of the state
+	int32 SelectStatus = 0;
+	switch (State)
+	{
+		case EIPv6SocketInternalState::CanRead:
+			SelectStatus = select(Socket + 1, &SocketSet, NULL, NULL, &Time);
+			break;
+		case EIPv6SocketInternalState::CanWrite:
+			SelectStatus = select(Socket + 1, NULL, &SocketSet, NULL, &Time);
+			break;
+		case EIPv6SocketInternalState::HasError:
+			SelectStatus = select(Socket + 1, NULL, NULL, &SocketSet, &Time);
+			break;
+	}
+
+	// if the select returns a positive number, the socket had the state, 0 means didn't have it, and negative is API error condition (not socket's error state)
+	return SelectStatus > 0  ? EIPv6SocketInternalState::Yes : 
+		   SelectStatus == 0 ? EIPv6SocketInternalState::No : 
+		   EIPv6SocketInternalState::EncounteredError;
+#else
+	UE_LOG(LogSockets, Fatal, TEXT("This platform doesn't support select(), but FSocketBSD::HasState was not overridden"));
+	return EIPv6SocketInternalState::EncounteredError;
+#endif
+}
+
+bool FSocketBSDIPv6::HasPendingConnection(bool& bHasPendingConnection)
+{
+	bool bHasSucceeded = false;
+	bHasPendingConnection = false;
+
+	// make sure socket has no error state
+	if (HasState(EIPv6SocketInternalState::HasError) == EIPv6SocketInternalState::No)
+	{
+		// get the read state
+		EIPv6SocketInternalState::Return State = HasState(EIPv6SocketInternalState::CanRead);
+		
+		// turn the result into the outputs
+		bHasSucceeded = State != EIPv6SocketInternalState::EncounteredError;
+		bHasPendingConnection = State == EIPv6SocketInternalState::Yes;
+	}
+
+	return bHasSucceeded;
+}
+
+
+bool FSocketBSDIPv6::HasPendingData(uint32& PendingDataSize)
+{
+	bool bHasSucceeded = false;
+	PendingDataSize = 0;
+
+	// make sure socket has no error state
+	if (HasState(EIPv6SocketInternalState::CanRead) == EIPv6SocketInternalState::Yes)
+	{
+#if PLATFORM_HAS_BSD_SOCKET_FEATURE_IOCTL
+		// See if there is any pending data on the read socket
+		if(ioctlsocket(Socket, FIONREAD, (u_long*)(&PendingDataSize)) == 0)
+#endif
+		{
+			bHasSucceeded = true;
+		}
+	}
+
+	return bHasSucceeded;
+}
+
+
+FSocket* FSocketBSDIPv6::Accept(const FString& SocketDescription)
+{
+	SOCKET NewSocket = accept(Socket,NULL,NULL);
+
+	if (NewSocket != INVALID_SOCKET)
+	{
+		// we need the subclass to create the actual FSocket object
+		check(SocketSubsystem);
+		FSocketSubsystemBSDIPv6* BSDSystem = static_cast<FSocketSubsystemBSDIPv6*>(SocketSubsystem);
+		return BSDSystem->InternalBSDSocketFactory(NewSocket, SocketType, SocketDescription);
+	}
+
+	return NULL;
+}
+
+
+FSocket* FSocketBSDIPv6::Accept(FInternetAddr& OutAddr, const FString& SocketDescription)
+{
+	SOCKLEN SizeOf = sizeof(sockaddr_in6);
+	SOCKET NewSocket = accept(Socket, *(FInternetAddrBSDIPv6*)(&OutAddr), &SizeOf);
+
+	if (NewSocket != INVALID_SOCKET)
+	{
+		// we need the subclass to create the actual FSocket object
+		check(SocketSubsystem);
+		FSocketSubsystemBSDIPv6* BSDSystem = static_cast<FSocketSubsystemBSDIPv6*>(SocketSubsystem);
+		return BSDSystem->InternalBSDSocketFactory(NewSocket, SocketType, SocketDescription);
+	}
+
+	return NULL;
+}
+
+
+bool FSocketBSDIPv6::SendTo(const uint8* Data, int32 Count, int32& BytesSent, const FInternetAddr& Destination)
+{
+	// Write the data and see how much was written
+	sockaddr* ToAddr = (FInternetAddrBSDIPv6&)Destination;
+	sockaddr_in6* ToAddr6 = (sockaddr_in6*)ToAddr;
+
+	BytesSent = sendto(Socket, (const char*)Data, Count, 0, (FInternetAddrBSDIPv6&)Destination, sizeof(sockaddr_in6));
+	if(BytesSent == SOCKET_ERROR)
+	{
+		int SockError = WSAGetLastError();
+		UE_LOG(LogSockets, Log, TEXT("sendto error: %d"), SockError);
+	}
+//	NETWORK_PROFILER(FSocket::SendTo(Data,Count,BytesSent,Destination));
+
+	return BytesSent >= 0;
+}
+
+
+bool FSocketBSDIPv6::Send(const uint8* Data, int32 Count, int32& BytesSent)
+{
+	BytesSent = send(Socket,(const char*)Data,Count,0);
+
+//	NETWORK_PROFILER(FSocket::Send(Data,Count,BytesSent));
+
+	return BytesSent >= 0;
+}
+
+
+bool FSocketBSDIPv6::RecvFrom(uint8* Data, int32 BufferSize, int32& BytesRead, FInternetAddr& Source, ESocketReceiveFlags::Type Flags)
+{
+	SOCKLEN Size = sizeof(sockaddr_in6);
+	sockaddr& Addr = *(FInternetAddrBSDIPv6&)Source;
+
+	// Read into the buffer and set the source address
+	BytesRead = recvfrom(Socket, (char*)Data, BufferSize, Flags, &Addr, &Size);
+
+//	NETWORK_PROFILER(FSocket::RecvFrom(Data,BufferSize,BytesRead,Source));
+
+	return BytesRead >= 0;
+}
+
+
+bool FSocketBSDIPv6::Recv(uint8* Data, int32 BufferSize, int32& BytesRead, ESocketReceiveFlags::Type Flags)
+{
+	BytesRead = recv(Socket, (char*)Data, BufferSize, Flags);
+
+//	NETWORK_PROFILER(FSocket::Recv(Data,BufferSize,BytesRead));
+
+	return BytesRead >= 0;
+}
+
+
+bool FSocketBSDIPv6::Wait(ESocketWaitConditions::Type Condition, FTimespan WaitTime)
+{
+	if ((Condition == ESocketWaitConditions::WaitForRead) || (Condition == ESocketWaitConditions::WaitForReadOrWrite))
+	{
+		if (HasState(EIPv6SocketInternalState::CanRead, WaitTime) == EIPv6SocketInternalState::Yes)
+		{
+			return true;
+		}
+	}
+
+	if ((Condition == ESocketWaitConditions::WaitForWrite) || (Condition == ESocketWaitConditions::WaitForReadOrWrite))
+	{
+		if (HasState(EIPv6SocketInternalState::CanWrite, WaitTime) == EIPv6SocketInternalState::Yes)
+		{
+			return true;
+		}
+	}
+
+	return false;
+}
+
+
+ESocketConnectionState FSocketBSDIPv6::GetConnectionState(void)
+{
+	ESocketConnectionState CurrentState = SCS_ConnectionError;
+
+	// look for an existing error
+	if (HasState(EIPv6SocketInternalState::HasError) == EIPv6SocketInternalState::No)
+	{
+		// get the write state
+		EIPv6SocketInternalState::Return State = HasState(EIPv6SocketInternalState::CanWrite);
+		
+		// translate yes or no (error is already set)
+		if (State == EIPv6SocketInternalState::Yes)
+		{
+			CurrentState = SCS_Connected;
+		}
+		else if (State == EIPv6SocketInternalState::No)
+		{
+			CurrentState = SCS_NotConnected;
+		}
+	}
+
+	return CurrentState;
+}
+
+
+void FSocketBSDIPv6::GetAddress(FInternetAddr& OutAddr)
+{
+	FInternetAddrBSDIPv6& Addr = (FInternetAddrBSDIPv6&)OutAddr;
+	SOCKLEN Size = sizeof(sockaddr_in6);
+
+	// Figure out what ip/port we are bound to
+	bool bOk = getsockname(Socket, Addr, &Size) == 0;
+
+	if (bOk == false)
+	{
+		check(SocketSubsystem);
+		UE_LOG(LogSockets, Error, TEXT("Failed to read address for socket (%s)"), SocketSubsystem->GetSocketError());
+	}
+}
+
+
+bool FSocketBSDIPv6::SetNonBlocking(bool bIsNonBlocking)
+{
+	u_long Value = bIsNonBlocking ? true : false;
+
+#if PLATFORM_HAS_BSD_SOCKET_FEATURE_WINSOCKETS
+	return ioctlsocket(Socket,FIONBIO,&Value) == 0;
+#else
+	int Flags = fcntl(Socket, F_GETFL, 0);
+	//Set the flag or clear it, without destroying the other flags.
+	Flags = bIsNonBlocking ? Flags | O_NONBLOCK : Flags ^ (Flags & O_NONBLOCK);
+	int err = fcntl(Socket, F_SETFL, Flags);
+	return (err == 0 ? true : false);
+#endif
+}
+
+
+bool FSocketBSDIPv6::SetBroadcast(bool bAllowBroadcast)
+{
+	int Param = bAllowBroadcast ? 1 : 0;
+	return setsockopt(Socket,SOL_SOCKET,SO_BROADCAST,(char*)&Param,sizeof(Param)) == 0;
+}
+
+
+bool FSocketBSDIPv6::JoinMulticastGroup(const FInternetAddr& GroupAddress)
+{
+	ipv6_mreq imr;
+
+	imr.ipv6mr_interface = 0;
+	imr.ipv6mr_multiaddr = ((sockaddr_in6*)&**((FInternetAddrBSDIPv6*)(&GroupAddress)))->sin6_addr;
+
+	return (setsockopt(Socket, IPPROTO_IPV6, IP_ADD_MEMBERSHIP, (char*)&imr, sizeof(imr)) == 0);
+}
+
+
+bool FSocketBSDIPv6::LeaveMulticastGroup(const FInternetAddr& GroupAddress)
+{
+	ipv6_mreq imr;
+
+	imr.ipv6mr_interface = 0;
+	imr.ipv6mr_multiaddr = ((sockaddr_in6*)&**((FInternetAddrBSDIPv6*)(&GroupAddress)))->sin6_addr;
+
+	return (setsockopt(Socket, IPPROTO_IPV6, IP_DROP_MEMBERSHIP, (char*)&imr, sizeof(imr)) == 0);
+}
+
+
+bool FSocketBSDIPv6::SetMulticastLoopback(bool bLoopback)
+{
+	return (setsockopt(Socket, IPPROTO_IPV6, IP_MULTICAST_LOOP, (char*)&bLoopback, sizeof(bLoopback)) == 0);
+}
+
+
+bool FSocketBSDIPv6::SetMulticastTtl(uint8 TimeToLive)
+{
+	return (setsockopt(Socket, IPPROTO_IPV6, IP_MULTICAST_TTL, (char*)&TimeToLive, sizeof(TimeToLive)) == 0);
+}
+
+
+bool FSocketBSDIPv6::SetReuseAddr(bool bAllowReuse)
+{
+	int Param = bAllowReuse ? 1 : 0;
+	return setsockopt(Socket,SOL_SOCKET,SO_REUSEADDR,(char*)&Param,sizeof(Param)) == 0;
+}
+
+
+bool FSocketBSDIPv6::SetLinger(bool bShouldLinger,int32 Timeout)
+{
+	linger ling;
+
+	ling.l_onoff = bShouldLinger;
+	ling.l_linger = Timeout;
+
+	return setsockopt(Socket,SOL_SOCKET,SO_LINGER,(char*)&ling,sizeof(ling)) == 0;
+}
+
+
+bool FSocketBSDIPv6::SetRecvErr(bool bUseErrorQueue)
+{
+	// Not supported, but return true to avoid spurious log messages
+	return true;
+}
+
+
+bool FSocketBSDIPv6::SetSendBufferSize(int32 Size,int32& NewSize)
+{
+	SOCKLEN SizeSize = sizeof(int32);
+	bool bOk = setsockopt(Socket,SOL_SOCKET,SO_SNDBUF,(char*)&Size,sizeof(int32)) == 0;
+
+	// Read the value back in case the size was modified
+	getsockopt(Socket,SOL_SOCKET,SO_SNDBUF,(char*)&NewSize, &SizeSize);
+
+	return bOk;
+}
+
+
+bool FSocketBSDIPv6::SetReceiveBufferSize(int32 Size,int32& NewSize)
+{
+	SOCKLEN SizeSize = sizeof(int32);
+	bool bOk = setsockopt(Socket,SOL_SOCKET,SO_RCVBUF,(char*)&Size,sizeof(int32)) == 0;
+
+	// Read the value back in case the size was modified
+	getsockopt(Socket,SOL_SOCKET,SO_RCVBUF,(char*)&NewSize, &SizeSize);
+
+	return bOk;
+}
+
+
+int32 FSocketBSDIPv6::GetPortNo(void)
+{
+	sockaddr_in6 Addr;
+
+	SOCKLEN Size = sizeof(sockaddr_in6);
+
+	// Figure out what ip/port we are bound to
+	bool bOk = getsockname(Socket, (sockaddr*)&Addr, &Size) == 0;
+	
+	if (bOk == false)
+	{
+		check(SocketSubsystem);
+		UE_LOG(LogSockets, Error, TEXT("Failed to read address for socket (%s)"), SocketSubsystem->GetSocketError());
+	}
+
+	// Read the port number
+	return ntohs(Addr.sin6_port);
+}
+
+bool FSocketBSDIPv6::SetIPv6Only(bool bIPv6Only)
+{
+	int v6only = bIPv6Only ? 1 : 0;
+	bool bOk = setsockopt(Socket,IPPROTO_IPV6,IPV6_V6ONLY,(char*) &v6only,sizeof( v6only )) == 0;
+
+	if(bOk == false)
+	{
+		check(SocketSubsystem);
+		UE_LOG(LogSockets, Error, TEXT("Failed to set sock opt for socket (%s)"), SocketSubsystem->GetSocketError());
+	}
+
+	return bOk;
+}
+
+#endif	//PLATFORM_HAS_BSD_SOCKETS
diff --git a/Engine/Source/Runtime/Sockets/Private/BSDIPv6Sockets/SocketsBSDIPv6.h b/Engine/Source/Runtime/Sockets/Private/BSDIPv6Sockets/SocketsBSDIPv6.h
new file mode 100644
--- /dev/null
+++ b/Engine/Source/Runtime/Sockets/Private/BSDIPv6Sockets/SocketsBSDIPv6.h
@@ -0,0 +1,132 @@
+// Copyright 1998-2014 Epic Games, Inc. All Rights Reserved.
+
+#pragma once
+
+#include "Platform.h"
+
+#if PLATFORM_HAS_BSD_SOCKETS
+
+#include "Sockets.h"
+#include "IPAddress.h"
+
+namespace EIPv6SocketInternalState
+{
+	enum Param
+	{
+		CanRead,
+		CanWrite,
+		HasError,
+	};
+	enum Return
+	{
+		Yes,
+		No,
+		EncounteredError,
+	};
+};
+
+/**
+ * Implements a BSD network socket.
+ */
+class FSocketBSDIPv6
+	: public FSocket
+{
+public:
+
+	/**
+	 * Assigns a BSD socket to this object
+	 *
+	 * @param InSocket the socket to assign to this object
+	 * @param InSocketType the type of socket that was created
+	 * @param InSocketDescription the debug description of the socket
+	 * @param InSubsystem the subsystem that created this socket
+	 */
+	FSocketBSDIPv6(SOCKET InSocket, ESocketType InSocketType, const FString& InSocketDescription, ISocketSubsystem* InSubsystem) 
+		: FSocket(InSocketType, InSocketDescription)
+		, Socket(InSocket)
+		, SocketSubsystem(InSubsystem)
+	{ }
+
+	/**
+	 * Destructor.
+	 *
+	 * Closes the socket if it is still open
+	 */
+	virtual ~FSocketBSDIPv6()
+	{
+		Close();
+	}
+
+
+public:
+
+	virtual bool Close() OVERRIDE;
+
+	virtual bool Bind(const FInternetAddr& Addr) OVERRIDE;
+
+	virtual bool Connect(const FInternetAddr& Addr) OVERRIDE;
+
+	virtual bool Listen(int32 MaxBacklog) OVERRIDE;
+
+	virtual bool HasPendingConnection(bool& bHasPendingConnection) OVERRIDE;
+
+	virtual bool HasPendingData(uint32& PendingDataSize) OVERRIDE;
+
+	virtual class FSocket* Accept(const FString& SocketDescription) OVERRIDE;
+
+	virtual class FSocket* Accept(FInternetAddr& OutAddr, const FString& SocketDescription) OVERRIDE;
+
+	virtual bool SendTo(const uint8* Data, int32 Count, int32& BytesSent, const FInternetAddr& Destination) OVERRIDE;
+
+	virtual bool Send(const uint8* Data, int32 Count, int32& BytesSent) OVERRIDE;
+
+	virtual bool RecvFrom(uint8* Data, int32 BufferSize, int32& BytesRead, FInternetAddr& Source, ESocketReceiveFlags::Type Flags = ESocketReceiveFlags::None) OVERRIDE;
+
+	virtual bool Recv(uint8* Data,int32 BufferSize,int32& BytesRead, ESocketReceiveFlags::Type Flags = ESocketReceiveFlags::None) OVERRIDE;
+
+	virtual bool Wait(ESocketWaitConditions::Type Condition, FTimespan WaitTime) OVERRIDE;
+
+	virtual ESocketConnectionState GetConnectionState() OVERRIDE;
+
+	virtual void GetAddress(FInternetAddr& OutAddr) OVERRIDE;
+
+	virtual bool SetNonBlocking(bool bIsNonBlocking = true) OVERRIDE;
+
+	virtual bool SetBroadcast(bool bAllowBroadcast = true) OVERRIDE;
+
+	virtual bool JoinMulticastGroup(const FInternetAddr& GroupAddress) OVERRIDE;
+
+	virtual bool LeaveMulticastGroup(const FInternetAddr& GroupAddress) OVERRIDE;
+
+	virtual bool SetMulticastLoopback(bool bLoopback) OVERRIDE;
+
+	virtual bool SetMulticastTtl(uint8 TimeToLive) OVERRIDE;
+
+	virtual bool SetReuseAddr(bool bAllowReuse = true) OVERRIDE;
+
+	virtual bool SetLinger(bool bShouldLinger = true, int32 Timeout = 0) OVERRIDE;
+
+	virtual bool SetRecvErr(bool bUseErrorQueue = true) OVERRIDE;
+
+	virtual bool SetSendBufferSize(int32 Size,int32& NewSize) OVERRIDE;
+
+	virtual bool SetReceiveBufferSize(int32 Size,int32& NewSize) OVERRIDE;
+
+	virtual int32 GetPortNo() OVERRIDE;
+
+	bool SetIPv6Only(bool bIPv6Only);
+
+
+protected:
+
+	// This is generally select(), but makes it easier for platforms without select to replace it
+	virtual EIPv6SocketInternalState::Return HasState(EIPv6SocketInternalState::Param State, FTimespan WaitTime=FTimespan(0));
+
+	// Holds the BSD socket object.
+	SOCKET Socket;
+
+	// Pointer to the subsystem that created it
+	ISocketSubsystem * SocketSubsystem;
+};
+
+#endif	//PLATFORM_HAS_BSD_SOCKETS
