#UE4 Fix for IPv6 sockets compile errors on non-Microsoft platforms

diff --git a/Engine/Source/Runtime/Core/Public/HAL/Platform.h b/Engine/Source/Runtime/Core/Public/HAL/Platform.h
--- a/Engine/Source/Runtime/Core/Public/HAL/Platform.h
+++ b/Engine/Source/Runtime/Core/Public/HAL/Platform.h
@@ -149,6 +149,9 @@
 #ifndef PLATFORM_HAS_BSD_SOCKETS
 	#define PLATFORM_HAS_BSD_SOCKETS			1
 #endif
+#ifndef PLATFORM_HAS_BSD_IPV6_SOCKETS
+	#define PLATFORM_HAS_BSD_IPV6_SOCKETS			0
+#endif
 #ifndef PLATFORM_USE_PTHREADS
 	#define PLATFORM_USE_PTHREADS				1
 #endif
diff --git a/Engine/Source/Runtime/Sockets/Private/BSDIPv6Sockets/IPAddressBSDIPv6.h b/Engine/Source/Runtime/Sockets/Private/BSDIPv6Sockets/IPAddressBSDIPv6.h
--- a/Engine/Source/Runtime/Sockets/Private/BSDIPv6Sockets/IPAddressBSDIPv6.h
+++ b/Engine/Source/Runtime/Sockets/Private/BSDIPv6Sockets/IPAddressBSDIPv6.h
@@ -2,7 +2,7 @@
 
 #pragma once
 
-#if PLATFORM_HAS_BSD_SOCKETS
+#if PLATFORM_HAS_BSD_IPV6_SOCKETS
 
 #include "Core.h"
 #include "IPAddress.h"
@@ -266,4 +266,4 @@ public:
  	}
 };
 
-#endif	//PLATFORM_HAS_BSD_SOCKETS
+#endif
diff --git a/Engine/Source/Runtime/Sockets/Private/BSDIPv6Sockets/SocketSubsystemBSDIPv6.cpp b/Engine/Source/Runtime/Sockets/Private/BSDIPv6Sockets/SocketSubsystemBSDIPv6.cpp
--- a/Engine/Source/Runtime/Sockets/Private/BSDIPv6Sockets/SocketSubsystemBSDIPv6.cpp
+++ b/Engine/Source/Runtime/Sockets/Private/BSDIPv6Sockets/SocketSubsystemBSDIPv6.cpp
@@ -2,7 +2,7 @@
 
 #include "SocketsPrivatePCH.h"
 
-#if PLATFORM_HAS_BSD_SOCKETS
+#if PLATFORM_HAS_BSD_IPV6_SOCKETS
 
 #include "SocketSubsystemBSDIPv6.h"
 #include "SocketsBSDIPv6.h"
@@ -194,4 +194,4 @@ ESocketErrors FSocketSubsystemBSDIPv6::TranslateErrorCode(int32 Code)
 	return SE_NO_ERROR;
 }
 
-#endif	//PLATFORM_HAS_BSD_SOCKETS
+#endif
diff --git a/Engine/Source/Runtime/Sockets/Private/BSDIPv6Sockets/SocketSubsystemBSDIPv6.h b/Engine/Source/Runtime/Sockets/Private/BSDIPv6Sockets/SocketSubsystemBSDIPv6.h
--- a/Engine/Source/Runtime/Sockets/Private/BSDIPv6Sockets/SocketSubsystemBSDIPv6.h
+++ b/Engine/Source/Runtime/Sockets/Private/BSDIPv6Sockets/SocketSubsystemBSDIPv6.h
@@ -5,7 +5,7 @@
 #include "SocketSubsystem.h"
 #include "SocketSubsystemPackage.h"
 
-#if PLATFORM_HAS_BSD_SOCKETS
+#if PLATFORM_HAS_BSD_IPV6_SOCKETS
 
 #include "../BSDSockets/SocketSubsystemBSDPrivate.h"
 #include "IPAddressBSDIPv6.h"
diff --git a/Engine/Source/Runtime/Sockets/Private/BSDIPv6Sockets/SocketsBSDIPv6.cpp b/Engine/Source/Runtime/Sockets/Private/BSDIPv6Sockets/SocketsBSDIPv6.cpp
--- a/Engine/Source/Runtime/Sockets/Private/BSDIPv6Sockets/SocketsBSDIPv6.cpp
+++ b/Engine/Source/Runtime/Sockets/Private/BSDIPv6Sockets/SocketsBSDIPv6.cpp
@@ -2,7 +2,7 @@
 
 #include "SocketsPrivatePCH.h"
 
-#if PLATFORM_HAS_BSD_SOCKETS
+#if PLATFORM_HAS_BSD_IPV6_SOCKETS
 
 #include "SocketSubsystemBSDIPv6.h"
 #include "SocketsBSDIPv6.h"
@@ -404,4 +404,4 @@ bool FSocketBSDIPv6::SetIPv6Only(bool bIPv6Only)
 	return bOk;
 }
 
-#endif	//PLATFORM_HAS_BSD_SOCKETS
+#endif
diff --git a/Engine/Source/Runtime/Sockets/Private/BSDIPv6Sockets/SocketsBSDIPv6.h b/Engine/Source/Runtime/Sockets/Private/BSDIPv6Sockets/SocketsBSDIPv6.h
--- a/Engine/Source/Runtime/Sockets/Private/BSDIPv6Sockets/SocketsBSDIPv6.h
+++ b/Engine/Source/Runtime/Sockets/Private/BSDIPv6Sockets/SocketsBSDIPv6.h
@@ -4,7 +4,7 @@
 
 #include "Platform.h"
 
-#if PLATFORM_HAS_BSD_SOCKETS
+#if PLATFORM_HAS_BSD_IPV6_SOCKETS
 
 #include "Sockets.h"
 #include "IPAddress.h"
@@ -129,4 +129,4 @@ protected:
 	ISocketSubsystem * SocketSubsystem;
 };
 
-#endif	//PLATFORM_HAS_BSD_SOCKETS
+#endif
