Making IPv6 socket code RFC 2553-compliant (and thus not Windows-specific). Enabling Linux to support BSD v6 sockets too.

diff --git a/Engine/Source/Runtime/Core/Public/Linux/LinuxPlatform.h b/Engine/Source/Runtime/Core/Public/Linux/LinuxPlatform.h
--- a/Engine/Source/Runtime/Core/Public/Linux/LinuxPlatform.h
+++ b/Engine/Source/Runtime/Core/Public/Linux/LinuxPlatform.h
@@ -42,6 +42,7 @@ typedef FLinuxPlatformTypes FPlatformTypes;
 #define PLATFORM_MAX_FILEPATH_LENGTH				MAX_PATH /* @todo linux: avoid using PATH_MAX as it is known to be broken */
 #define PLATFORM_HAS_NO_EPROCLIM					1
 #define PLATFORM_HAS_BSD_SOCKET_FEATURE_IOCTL		1
+#define PLATFORM_HAS_BSD_IPV6_SOCKETS				1
 #define PLATFORM_SUPPORTS_JEMALLOC					1
 #define PLATFORM_EXCEPTIONS_DISABLED				1
 
diff --git a/Engine/Source/Runtime/Sockets/Private/BSDIPv6Sockets/IPAddressBSDIPv6.h b/Engine/Source/Runtime/Sockets/Private/BSDIPv6Sockets/IPAddressBSDIPv6.h
--- a/Engine/Source/Runtime/Sockets/Private/BSDIPv6Sockets/IPAddressBSDIPv6.h
+++ b/Engine/Source/Runtime/Sockets/Private/BSDIPv6Sockets/IPAddressBSDIPv6.h
@@ -49,7 +49,7 @@ public:
 		else
 		{
 			in_addr Addr;
-			Addr.S_un.S_addr= htonl(InAddr);
+			Addr.s_addr= htonl(InAddr);
 
 			SetIp(Addr);
 		}
@@ -64,7 +64,8 @@ public:
 	virtual void SetIp(const TCHAR* InAddr, bool& bIsValid) OVERRIDE
 	{
 		// check for valid IPv6 address
-		if(InetPtonW(AF_INET6, InAddr, &Addr.sin6_addr))
+		auto InAddrAnsi = StringCast<ANSICHAR>(InAddr);
+		if (inet_pton(AF_INET6, InAddrAnsi.Get(), &Addr.sin6_addr))
 		{
 			bIsValid = true;
 			return;
@@ -79,7 +80,7 @@ public:
 			// Remove the closing brace if it exists
 			CopiedAddr.RemoveFromEnd(TEXT("]"));
 
-			if(InetPtonW(AF_INET6, CopiedAddr.GetCharArray().GetTypedData(), &Addr.sin6_addr))
+			if(inet_pton(AF_INET6, StringCast<ANSICHAR>(CopiedAddr.GetCharArray().GetTypedData()).Get(), &Addr.sin6_addr))
 			{
 				bIsValid = true;
 				return;
@@ -88,7 +89,7 @@ public:
 
 		// Check if it's a valid IPv4 address, and if it is convert
 		in_addr  IPv4Addr;
-		if(InetPtonW(AF_INET, InAddr, &IPv4Addr))
+		if (inet_pton(AF_INET, InAddrAnsi.Get(), &IPv4Addr))
 		{
 			bIsValid = true;
 			SetIp(IPv4Addr);
@@ -109,13 +110,19 @@ public:
 		FMemory::Memzero(&Addr.sin6_addr,sizeof(Addr.sin6_addr));
 
 		// special mapping of ipv4 to ipv6 using a hybrid stack, won't work on a pure ipv6 implementation
-		Addr.sin6_addr.u.Word[5] = 0xffff;
-		Addr.sin6_addr.u.Byte[12] = IPv4Addr.S_un.S_un_b.s_b1;
-		Addr.sin6_addr.u.Byte[13] = IPv4Addr.S_un.S_un_b.s_b2;
-		Addr.sin6_addr.u.Byte[14] = IPv4Addr.S_un.S_un_b.s_b3;
-		Addr.sin6_addr.u.Byte[15] = IPv4Addr.S_un.S_un_b.s_b4;
-
-		UE_LOG(LogSockets, Log, TEXT("Using IPv4 address: %d.%d.%d.%d  on an ipv6 socket"), IPv4Addr.S_un.S_un_b.s_b1, IPv4Addr.S_un.S_un_b.s_b2, IPv4Addr.S_un.S_un_b.s_b3, IPv4Addr.S_un.S_un_b.s_b4);
+		uint8	IPv4b1 = (static_cast<uint32>(IPv4Addr.s_addr) & 0xFF),
+				IPv4b2 = ((static_cast<uint32>(IPv4Addr.s_addr) >> 8) & 0xFF),
+				IPv4b3 = ((static_cast<uint32>(IPv4Addr.s_addr) >> 16) & 0xFF),
+				IPv4b4 = ((static_cast<uint32>(IPv4Addr.s_addr) >> 24) & 0xFF);
+		
+		Addr.sin6_addr.s6_addr[10] = 0xff;
+		Addr.sin6_addr.s6_addr[11] = 0xff;
+		Addr.sin6_addr.s6_addr[12] = IPv4b1;
+		Addr.sin6_addr.s6_addr[13] = IPv4b2;
+		Addr.sin6_addr.s6_addr[14] = IPv4b3;
+		Addr.sin6_addr.s6_addr[15] = IPv4b4;
+
+		UE_LOG(LogSockets, Log, TEXT("Using IPv4 address: %d.%d.%d.%d  on an ipv6 socket"), IPv4b1, IPv4b2, IPv4b3, IPv4b4);
 	}
 
 	/**
@@ -158,7 +165,10 @@ public:
 		// in a dual stack system.
 		// This function doesn't really make sense in IPv6, but too much other code relies on it
 		// existing to not have this here.
-		OutAddr = (Addr.sin6_addr.u.Word[6] << 16) & (Addr.sin6_addr.u.Word[7]);
+
+		//OutAddr = (Addr.sin6_addr.u.Word[6] << 16) & (Addr.sin6_addr.u.Word[7]);
+		// FIXME [RCL]: original code had & here (see commented line above), which looks like a bug. 
+		OutAddr = (Addr.sin6_addr.s6_addr[12] << 24) | (Addr.sin6_addr.s6_addr[13] << 16) | (Addr.sin6_addr.s6_addr[14] << 8) | (Addr.sin6_addr.s6_addr[15]);
 	}
 #if 0
 	/**
@@ -208,6 +218,14 @@ public:
 	virtual void SetBroadcastAddress() OVERRIDE
 	{
 		// broadcast means something different in IPv6, but this is a rough equivalent
+#ifndef in6addr_allnodesonlink
+		// see RFC 4291, link-local multicast address http://tools.ietf.org/html/rfc4291
+		static in6_addr in6addr_allnodesonlink =
+		{
+			{ { 0xff, 0x02, 0x00, 0x00, 0x00, 0x00, 0x00, 0x00, 0x00, 0x00, 0x00, 0x00, 0x00, 0x00, 0x00, 0x01 } }
+		};
+#endif // in6addr_allnodesonlink
+
 		SetIp(in6addr_allnodesonlink);
 		SetPort(0);
 	}
diff --git a/Engine/Source/Runtime/Sockets/Private/BSDIPv6Sockets/SocketSubsystemBSDIPv6.cpp b/Engine/Source/Runtime/Sockets/Private/BSDIPv6Sockets/SocketSubsystemBSDIPv6.cpp
--- a/Engine/Source/Runtime/Sockets/Private/BSDIPv6Sockets/SocketSubsystemBSDIPv6.cpp
+++ b/Engine/Source/Runtime/Sockets/Private/BSDIPv6Sockets/SocketSubsystemBSDIPv6.cpp
@@ -62,7 +62,7 @@ ESocketErrors FSocketSubsystemBSDIPv6::GetHostByName(const ANSICHAR* HostName, F
 			// we are copying they will get munged data. This relies on the
 			// consumer of this class to call the resolved() accessor before
 			// attempting to read this data
-			((FInternetAddrBSDIPv6&)OutAddr).SetIp(*(in_addr6*)(*HostEnt->h_addr_list));
+			((FInternetAddrBSDIPv6&)OutAddr).SetIp(*(in6_addr*)(*HostEnt->h_addr_list));
 		}
 		else
 		{
diff --git a/Engine/Source/Runtime/Sockets/Private/BSDIPv6Sockets/SocketsBSDIPv6.cpp b/Engine/Source/Runtime/Sockets/Private/BSDIPv6Sockets/SocketsBSDIPv6.cpp
--- a/Engine/Source/Runtime/Sockets/Private/BSDIPv6Sockets/SocketsBSDIPv6.cpp
+++ b/Engine/Source/Runtime/Sockets/Private/BSDIPv6Sockets/SocketsBSDIPv6.cpp
@@ -159,8 +159,8 @@ bool FSocketBSDIPv6::SendTo(const uint8* Data, int32 Count, int32& BytesSent, co
 	BytesSent = sendto(Socket, (const char*)Data, Count, 0, (FInternetAddrBSDIPv6&)Destination, sizeof(sockaddr_in6));
 	if(BytesSent == SOCKET_ERROR)
 	{
-		int SockError = WSAGetLastError();
-		UE_LOG(LogSockets, Log, TEXT("sendto error: %d"), SockError);
+		ESocketErrors SockError = SocketSubsystem->GetLastErrorCode();
+		UE_LOG(LogSockets, Log, TEXT("sendto error: (ESocketErrors:%d)"), SockError);
 	}
 //	NETWORK_PROFILER(FSocket::SendTo(Data,Count,BytesSent,Destination));
 
