Making the decoding of IPv4 address symmetric with encoding. This seems to be specific to Windows/Xbox dual-stack.

diff --git a/Engine/Source/Runtime/Sockets/Private/BSDIPv6Sockets/IPAddressBSDIPv6.h b/Engine/Source/Runtime/Sockets/Private/BSDIPv6Sockets/IPAddressBSDIPv6.h
--- a/Engine/Source/Runtime/Sockets/Private/BSDIPv6Sockets/IPAddressBSDIPv6.h
+++ b/Engine/Source/Runtime/Sockets/Private/BSDIPv6Sockets/IPAddressBSDIPv6.h
@@ -168,7 +168,7 @@ public:
 
 		//OutAddr = (Addr.sin6_addr.u.Word[6] << 16) & (Addr.sin6_addr.u.Word[7]);
 		// FIXME [RCL]: original code had & here (see commented line above), which looks like a bug. 
-		OutAddr = (Addr.sin6_addr.s6_addr[12] << 24) | (Addr.sin6_addr.s6_addr[13] << 16) | (Addr.sin6_addr.s6_addr[14] << 8) | (Addr.sin6_addr.s6_addr[15]);
+		OutAddr = (Addr.sin6_addr.s6_addr[15] << 24) | (Addr.sin6_addr.s6_addr[14] << 16) | (Addr.sin6_addr.s6_addr[13] << 8) | (Addr.sin6_addr.s6_addr[12]);
 	}
 #if 0
 	/**
