#ue4
#ios
* fix for shutdown crash on iOS, no longer try to stop the socket subsystem if we aren't on the main thread

TTP328922

diff --git a/Engine/Source/Runtime/Sockets/Private/SocketSubsystem.cpp b/Engine/Source/Runtime/Sockets/Private/SocketSubsystem.cpp
--- a/Engine/Source/Runtime/Sockets/Private/SocketSubsystem.cpp
+++ b/Engine/Source/Runtime/Sockets/Private/SocketSubsystem.cpp
@@ -70,7 +70,7 @@ static TSharedPtr<IModuleInterface> LoadSubsystemModule(const FString& Subsystem
  */
 void ISocketSubsystem::ShutdownAllSystems()
 {
-	if (FModuleManager::Get().IsModuleLoaded(TEXT("Sockets")) == true)
+	if (IsInGameThread() && FModuleManager::Get().IsModuleLoaded(TEXT("Sockets")) == true)
 	{
 		// Unloading the Sockets module will call FSocketSubsystemModule::ShutdownSocketSubsystem()
 		const bool bIsShutdown = true;
