Fix crash if a BSDIPv6 socket is double-closed
On close, the socket now sets itself to INVALID_SOCKET and won't try to close if it's already an INVALID_SOCKET (same as IPv4 BSD sockets)

diff --git a/Engine/Source/Runtime/Sockets/Private/BSDIPv6Sockets/SocketsBSDIPv6.cpp b/Engine/Source/Runtime/Sockets/Private/BSDIPv6Sockets/SocketsBSDIPv6.cpp
--- a/Engine/Source/Runtime/Sockets/Private/BSDIPv6Sockets/SocketsBSDIPv6.cpp
+++ b/Engine/Source/Runtime/Sockets/Private/BSDIPv6Sockets/SocketsBSDIPv6.cpp
@@ -10,7 +10,13 @@
 
 bool FSocketBSDIPv6::Close(void)
 {
-	return closesocket(Socket) == 0;
+	if (Socket != INVALID_SOCKET)
+	{
+		int32 error = closesocket(Socket);
+		Socket = INVALID_SOCKET;
+		return error == 0;
+	}
+	return false;
 }
 
 
