Parse port numbers for IPv4 and IPv6 addresses.

diff --git a/Engine/Source/Runtime/Sockets/Private/BSDIPv6Sockets/IPAddressBSDIPv6.h b/Engine/Source/Runtime/Sockets/Private/BSDIPv6Sockets/IPAddressBSDIPv6.h
--- a/Engine/Source/Runtime/Sockets/Private/BSDIPv6Sockets/IPAddressBSDIPv6.h
+++ b/Engine/Source/Runtime/Sockets/Private/BSDIPv6Sockets/IPAddressBSDIPv6.h
@@ -63,28 +63,41 @@ public:
 	 */
 	virtual void SetIp(const TCHAR* InAddr, bool& bIsValid) override
 	{
-		// check for valid IPv6 address
-		auto InAddrAnsi = StringCast<ANSICHAR>(InAddr);
-		if (inet_pton(AF_INET6, InAddrAnsi.Get(), &Addr.sin6_addr))
+
+		// Count colons to determine IP address type
+		int NumColons = 0;
+		int LastColonIndex = 0;
+		int i = 0;
+		while(InAddr && InAddr[i] != '\0')
 		{
-			bIsValid = true;
-			return;
+			if(InAddr[i] == ':')
+			{
+				NumColons++;
+				LastColonIndex = i;
+			}
+			i++;
 		}
 
-		// IPv6 URLs are surrounded by square brackets, check for that.
-		if(InAddr && *InAddr == '[')
+		// Break off port portion of string
+		FString AddressString(InAddr);
+		FString Port;
+		if(NumColons == 1 || NumColons == 8)
 		{
-			// Make a copy, skipping the opening brace.
-			FString CopiedAddr(InAddr + 1);
+			Port = AddressString.RightChop(LastColonIndex+1);
+			AddressString = AddressString.Left(LastColonIndex);
+			SetPort(FCString::Atoi(*Port));
+		}
 
-			// Remove the closing brace if it exists
-			CopiedAddr.RemoveFromEnd(TEXT("]"));
+		// IPv6 URLs are surrounded by square brackets, check for that.
+		AddressString.RemoveFromStart(FString("["));
+		AddressString.RemoveFromEnd(FString("]"));
 
-			if(inet_pton(AF_INET6, StringCast<ANSICHAR>(CopiedAddr.GetCharArray().GetTypedData()).Get(), &Addr.sin6_addr))
-			{
-				bIsValid = true;
-				return;
-			}
+		// check for valid IPv6 address
+		auto InAddrAnsi = StringCast<ANSICHAR>(*AddressString);
+		if (inet_pton(AF_INET6, InAddrAnsi.Get(), &Addr.sin6_addr))
+		{
+			bIsValid = true;
+			return;
 		}
 
 		// Check if it's a valid IPv4 address, and if it is convert
