HTML5 Win32 FName compile fix

diff --git a/Engine/Source/Runtime/Sockets/Private/HTML5/Simulator/SocketSubsystem.cpp b/Engine/Source/Runtime/Sockets/Private/HTML5/Simulator/SocketSubsystem.cpp
--- a/Engine/Source/Runtime/Sockets/Private/HTML5/Simulator/SocketSubsystem.cpp
+++ b/Engine/Source/Runtime/Sockets/Private/HTML5/Simulator/SocketSubsystem.cpp
@@ -93,7 +93,7 @@ bool FSocketSubsystemHTML5::HasNetworkDevice()
 class FSocket* FSocketSubsystemHTML5::CreateSocket( const FName& SocketType, const FString& SocketDescription, bool bForceUDP /*= false */ )
 {
 	FSocketX* NewSocket = NULL;
-	switch (SocketType.GetIndex())
+	switch (SocketType.GetComparisonIndex())
 	{
 	case NAME_DGram:
 		// Creates a data gram (UDP) socket 
