Translate socket flags to handle platforms that define them with different values.

diff --git a/Engine/Source/Runtime/Sockets/Private/BSDIPv6Sockets/SocketsBSDIPv6.cpp b/Engine/Source/Runtime/Sockets/Private/BSDIPv6Sockets/SocketsBSDIPv6.cpp
--- a/Engine/Source/Runtime/Sockets/Private/BSDIPv6Sockets/SocketsBSDIPv6.cpp
+++ b/Engine/Source/Runtime/Sockets/Private/BSDIPv6Sockets/SocketsBSDIPv6.cpp
@@ -8,6 +8,28 @@
 #include "SocketsBSDIPv6.h"
 //#include "Net/NetworkProfiler.h"
 
+namespace
+{
+	// Since the flag constants may have different values per-platform, translate into corresponding system constants.
+	// For example, MSG_WAITALL is 0x8 on Windows, but 0x100 on other platforms.
+	int TranslateFlags(ESocketReceiveFlags::Type Flags)
+	{
+		int TranslatedFlags = 0;
+
+		if (Flags & ESocketReceiveFlags::Peek)
+		{
+			TranslatedFlags |= MSG_PEEK;
+		}
+
+		if (Flags & ESocketReceiveFlags::WaitAll)
+		{
+			TranslatedFlags |= MSG_WAITALL;
+		}
+
+		return TranslatedFlags;
+	}
+}
+
 bool FSocketBSDIPv6::Close(void)
 {
 	if (Socket != INVALID_SOCKET)
@@ -186,8 +208,10 @@ bool FSocketBSDIPv6::RecvFrom(uint8* Data, int32 BufferSize, int32& BytesRead, F
 	SOCKLEN Size = sizeof(sockaddr_in6);
 	sockaddr& Addr = *(FInternetAddrBSDIPv6&)Source;
 
+	const int TranslatedFlags = TranslateFlags(Flags);
+
 	// Read into the buffer and set the source address
-	BytesRead = recvfrom(Socket, (char*)Data, BufferSize, Flags, &Addr, &Size);
+	BytesRead = recvfrom(Socket, (char*)Data, BufferSize, TranslatedFlags, &Addr, &Size);
 
 //	NETWORK_PROFILER(FSocket::RecvFrom(Data,BufferSize,BytesRead,Source));
 
@@ -197,7 +221,8 @@ bool FSocketBSDIPv6::RecvFrom(uint8* Data, int32 BufferSize, int32& BytesRead, F
 
 bool FSocketBSDIPv6::Recv(uint8* Data, int32 BufferSize, int32& BytesRead, ESocketReceiveFlags::Type Flags)
 {
-	BytesRead = recv(Socket, (char*)Data, BufferSize, Flags);
+	const int TranslatedFlags = TranslateFlags(Flags);
+	BytesRead = recv(Socket, (char*)Data, BufferSize, TranslatedFlags);
 
 //	NETWORK_PROFILER(FSocket::Recv(Data,BufferSize,BytesRead));
 
diff --git a/Engine/Source/Runtime/Sockets/Private/BSDSockets/SocketsBSD.cpp b/Engine/Source/Runtime/Sockets/Private/BSDSockets/SocketsBSD.cpp
--- a/Engine/Source/Runtime/Sockets/Private/BSDSockets/SocketsBSD.cpp
+++ b/Engine/Source/Runtime/Sockets/Private/BSDSockets/SocketsBSD.cpp
@@ -12,6 +12,28 @@
 typedef unsigned long   u_long;
 #endif 
 
+namespace
+{
+	// Since the flag constants may have different values per-platform, translate into corresponding system constants.
+	// For example, MSG_WAITALL is 0x8 on Windows, but 0x100 on other platforms.
+	int TranslateFlags(ESocketReceiveFlags::Type Flags)
+	{
+		int TranslatedFlags = 0;
+
+		if (Flags & ESocketReceiveFlags::Peek)
+		{
+			TranslatedFlags |= MSG_PEEK;
+		}
+
+		if (Flags & ESocketReceiveFlags::WaitAll)
+		{
+			TranslatedFlags |= MSG_WAITALL;
+		}
+
+		return TranslatedFlags;
+	}
+}
+
 /* FSocket interface
  *****************************************************************************/
 
@@ -198,8 +220,10 @@ bool FSocketBSD::RecvFrom(uint8* Data, int32 BufferSize, int32& BytesRead, FInte
 	SOCKLEN Size = sizeof(sockaddr_in);
 	sockaddr& Addr = *(FInternetAddrBSD&)Source;
 
+	const int TranslatedFlags = TranslateFlags(Flags);
+
 	// Read into the buffer and set the source address
-	BytesRead = recvfrom(Socket, (char*)Data, BufferSize, Flags, &Addr, &Size);
+	BytesRead = recvfrom(Socket, (char*)Data, BufferSize, TranslatedFlags, &Addr, &Size);
 //	NETWORK_PROFILER(FSocket::RecvFrom(Data,BufferSize,BytesRead,Source));
 
 	bool Result = BytesRead >= 0;
@@ -214,7 +238,8 @@ bool FSocketBSD::RecvFrom(uint8* Data, int32 BufferSize, int32& BytesRead, FInte
 
 bool FSocketBSD::Recv(uint8* Data, int32 BufferSize, int32& BytesRead, ESocketReceiveFlags::Type Flags)
 {
-	BytesRead = recv(Socket, (char*)Data, BufferSize, Flags);
+	const int TranslatedFlags = TranslateFlags(Flags);
+	BytesRead = recv(Socket, (char*)Data, BufferSize, TranslatedFlags);
 
 //	NETWORK_PROFILER(FSocket::Recv(Data,BufferSize,BytesRead));
 
