Place IPv6 TranslateFlags() in a different namespace.

- Fixes unity build when both IPv4 and IPv6 sockets happen to be in a single compile unit.

diff --git a/Engine/Source/Runtime/Sockets/Private/BSDIPv6Sockets/SocketsBSDIPv6.cpp b/Engine/Source/Runtime/Sockets/Private/BSDIPv6Sockets/SocketsBSDIPv6.cpp
--- a/Engine/Source/Runtime/Sockets/Private/BSDIPv6Sockets/SocketsBSDIPv6.cpp
+++ b/Engine/Source/Runtime/Sockets/Private/BSDIPv6Sockets/SocketsBSDIPv6.cpp
@@ -8,7 +8,7 @@
 #include "SocketsBSDIPv6.h"
 //#include "Net/NetworkProfiler.h"
 
-namespace
+namespace SocketBSDIPv6
 {
 	// Since the flag constants may have different values per-platform, translate into corresponding system constants.
 	// For example, MSG_WAITALL is 0x8 on Windows, but 0x100 on other platforms.
@@ -208,7 +208,7 @@ bool FSocketBSDIPv6::RecvFrom(uint8* Data, int32 BufferSize, int32& BytesRead, F
 	SOCKLEN Size = sizeof(sockaddr_in6);
 	sockaddr& Addr = *(FInternetAddrBSDIPv6&)Source;
 
-	const int TranslatedFlags = TranslateFlags(Flags);
+	const int TranslatedFlags = SocketBSDIPv6::TranslateFlags(Flags);
 
 	// Read into the buffer and set the source address
 	BytesRead = recvfrom(Socket, (char*)Data, BufferSize, TranslatedFlags, &Addr, &Size);
@@ -221,7 +221,7 @@ bool FSocketBSDIPv6::RecvFrom(uint8* Data, int32 BufferSize, int32& BytesRead, F
 
 bool FSocketBSDIPv6::Recv(uint8* Data, int32 BufferSize, int32& BytesRead, ESocketReceiveFlags::Type Flags)
 {
-	const int TranslatedFlags = TranslateFlags(Flags);
+	const int TranslatedFlags = SocketBSDIPv6::TranslateFlags(Flags);
 	BytesRead = recv(Socket, (char*)Data, BufferSize, TranslatedFlags);
 
 //	NETWORK_PROFILER(FSocket::Recv(Data,BufferSize,BytesRead));
