Combine the implementations of TranslateFlags.

diff --git a/Engine/Source/Runtime/Sockets/Private/BSDIPv6Sockets/SocketsBSDIPv6.cpp b/Engine/Source/Runtime/Sockets/Private/BSDIPv6Sockets/SocketsBSDIPv6.cpp
--- a/Engine/Source/Runtime/Sockets/Private/BSDIPv6Sockets/SocketsBSDIPv6.cpp
+++ b/Engine/Source/Runtime/Sockets/Private/BSDIPv6Sockets/SocketsBSDIPv6.cpp
@@ -8,28 +8,6 @@
 #include "SocketsBSDIPv6.h"
 //#include "Net/NetworkProfiler.h"
 
-namespace SocketBSDIPv6
-{
-	// Since the flag constants may have different values per-platform, translate into corresponding system constants.
-	// For example, MSG_WAITALL is 0x8 on Windows, but 0x100 on other platforms.
-	int TranslateFlags(ESocketReceiveFlags::Type Flags)
-	{
-		int TranslatedFlags = 0;
-
-		if (Flags & ESocketReceiveFlags::Peek)
-		{
-			TranslatedFlags |= MSG_PEEK;
-		}
-
-		if (Flags & ESocketReceiveFlags::WaitAll)
-		{
-			TranslatedFlags |= MSG_WAITALL;
-		}
-
-		return TranslatedFlags;
-	}
-}
-
 bool FSocketBSDIPv6::Close(void)
 {
 	if (Socket != INVALID_SOCKET)
@@ -208,7 +186,7 @@ bool FSocketBSDIPv6::RecvFrom(uint8* Data, int32 BufferSize, int32& BytesRead, F
 	SOCKLEN Size = sizeof(sockaddr_in6);
 	sockaddr& Addr = *(FInternetAddrBSDIPv6&)Source;
 
-	const int TranslatedFlags = SocketBSDIPv6::TranslateFlags(Flags);
+	const int TranslatedFlags = TranslateFlags(Flags);
 
 	// Read into the buffer and set the source address
 	BytesRead = recvfrom(Socket, (char*)Data, BufferSize, TranslatedFlags, &Addr, &Size);
@@ -221,7 +199,7 @@ bool FSocketBSDIPv6::RecvFrom(uint8* Data, int32 BufferSize, int32& BytesRead, F
 
 bool FSocketBSDIPv6::Recv(uint8* Data, int32 BufferSize, int32& BytesRead, ESocketReceiveFlags::Type Flags)
 {
-	const int TranslatedFlags = SocketBSDIPv6::TranslateFlags(Flags);
+	const int TranslatedFlags = TranslateFlags(Flags);
 	BytesRead = recv(Socket, (char*)Data, BufferSize, TranslatedFlags);
 
 //	NETWORK_PROFILER(FSocket::Recv(Data,BufferSize,BytesRead));
diff --git a/Engine/Source/Runtime/Sockets/Private/BSDSockets/SocketSubsystemBSDPrivate.h b/Engine/Source/Runtime/Sockets/Private/BSDSockets/SocketSubsystemBSDPrivate.h
--- a/Engine/Source/Runtime/Sockets/Private/BSDSockets/SocketSubsystemBSDPrivate.h
+++ b/Engine/Source/Runtime/Sockets/Private/BSDSockets/SocketSubsystemBSDPrivate.h
@@ -2,6 +2,8 @@
 
 #pragma once
 
+#include "SocketTypes.h"
+
 #if PLATFORM_HAS_BSD_SOCKET_FEATURE_WINSOCKETS
 	#include "AllowWindowsPlatformTypes.h"
 
@@ -40,4 +42,23 @@
 
 #endif
 
+// Since the flag constants may have different values per-platform, translate into corresponding system constants.
+// For example, MSG_WAITALL is 0x8 on Windows, but 0x100 on other platforms.
+inline int TranslateFlags(ESocketReceiveFlags::Type Flags)
+{
+	int TranslatedFlags = 0;
+
+	if (Flags & ESocketReceiveFlags::Peek)
+	{
+		TranslatedFlags |= MSG_PEEK;
+	}
+
+	if (Flags & ESocketReceiveFlags::WaitAll)
+	{
+		TranslatedFlags |= MSG_WAITALL;
+	}
+
+	return TranslatedFlags;
+}
+
 #include <errno.h>
diff --git a/Engine/Source/Runtime/Sockets/Private/BSDSockets/SocketsBSD.cpp b/Engine/Source/Runtime/Sockets/Private/BSDSockets/SocketsBSD.cpp
--- a/Engine/Source/Runtime/Sockets/Private/BSDSockets/SocketsBSD.cpp
+++ b/Engine/Source/Runtime/Sockets/Private/BSDSockets/SocketsBSD.cpp
@@ -12,28 +12,6 @@
 typedef unsigned long   u_long;
 #endif 
 
-namespace
-{
-	// Since the flag constants may have different values per-platform, translate into corresponding system constants.
-	// For example, MSG_WAITALL is 0x8 on Windows, but 0x100 on other platforms.
-	int TranslateFlags(ESocketReceiveFlags::Type Flags)
-	{
-		int TranslatedFlags = 0;
-
-		if (Flags & ESocketReceiveFlags::Peek)
-		{
-			TranslatedFlags |= MSG_PEEK;
-		}
-
-		if (Flags & ESocketReceiveFlags::WaitAll)
-		{
-			TranslatedFlags |= MSG_WAITALL;
-		}
-
-		return TranslatedFlags;
-	}
-}
-
 /* FSocket interface
  *****************************************************************************/
 
