Integrate final set of cert changes for Shootergame ps4.

diff --git a/Engine/Source/Runtime/Sockets/Private/BSDSockets/IPAddressBSD.h b/Engine/Source/Runtime/Sockets/Private/BSDSockets/IPAddressBSD.h
--- a/Engine/Source/Runtime/Sockets/Private/BSDSockets/IPAddressBSD.h
+++ b/Engine/Source/Runtime/Sockets/Private/BSDSockets/IPAddressBSD.h
@@ -12,6 +12,7 @@
  */
 class FInternetAddrBSD : public FInternetAddr
 {
+protected:
 	/** The internet ip address structure */
 	sockaddr_in Addr;
 
diff --git a/Engine/Source/Runtime/Sockets/Public/IPAddress.h b/Engine/Source/Runtime/Sockets/Public/IPAddress.h
--- a/Engine/Source/Runtime/Sockets/Public/IPAddress.h
+++ b/Engine/Source/Runtime/Sockets/Public/IPAddress.h
@@ -61,6 +61,22 @@ public:
 	 */
 	virtual int32 GetPort() const = 0;
 
+	/** 
+	 * Set Platform specific port data
+	 */
+	virtual void SetPlatformPort(int32 InPort)
+	{
+		SetPort(InPort);
+	}
+
+	/** 
+	 * Get platform specific port data.
+	 */
+	virtual int32 GetPlatformPort() const
+	{
+		return GetPort();
+	}
+
 	/** Sets the address to be any address */
 	virtual void SetAnyAddress() = 0;
 
