Fix UE-4390 - convert gethostbyname function calls to getaddrinfo function calls.

diff --git a/Engine/Source/Runtime/HTML5/Simulator/HTML5Win32/Private/SocketRaw.cpp b/Engine/Source/Runtime/HTML5/Simulator/HTML5Win32/Private/SocketRaw.cpp
--- a/Engine/Source/Runtime/HTML5/Simulator/HTML5Win32/Private/SocketRaw.cpp
+++ b/Engine/Source/Runtime/HTML5/Simulator/HTML5Win32/Private/SocketRaw.cpp
@@ -350,22 +350,27 @@ bool FSocketRaw::IsValid()
 
 bool FSocketRaw::GetHostByName( const char *NAME, FInternetAddrRaw& Address )
 {
-	hostent* HostEnt = gethostbyname(NAME);
-	// Make sure it's a valid type
-	if (HostEnt->h_addrtype == PF_INET)
+	addrinfo* AddrInfo = nullptr;
+
+	// We are only interested in IPv4 addresses.
+	addrinfo HintAddrInfo;
+	memset(&HintAddrInfo, 0, sizeof(HintAddrInfo));
+	HintAddrInfo.ai_family = AF_INET;
+
+	int ErrorCode = getaddrinfo(NAME, nullptr, &HintAddrInfo, &AddrInfo);
+	if (ErrorCode != 0)
 	{
-		// Copy the data before letting go of the lock. This is safe only
-		// for the copy locally. If another thread is reading this while
-		// we are copying they will get munged data. This relies on the
-		// consumer of this class to call the resolved() accessor before
-		// attempting to read this data
-		(Address).SetIp(*(unsigned int *)(*HostEnt->h_addr_list));
+		if (AddrInfo->ai_addr == nullptr)
+		{
+			return false;
+		}
+
+		sockaddr_in* IPv4Address = reinterpret_cast<sockaddr_in*>(AddrInfo->ai_addr);
+		unsigned int HostIP = ntohl(IPv4Address->sin_addr.s_addr);
+		Address.SetIp(HostIP);
 		return true;
 	}
-	else
-	{
-		return false; 
-	}
+	return false;
 }
 
 bool FSocketRaw::GetHostName(const char *NAME )
diff --git a/Engine/Source/Runtime/Sockets/Private/BSDIPv6Sockets/SocketSubsystemBSDIPv6.cpp b/Engine/Source/Runtime/Sockets/Private/BSDIPv6Sockets/SocketSubsystemBSDIPv6.cpp
--- a/Engine/Source/Runtime/Sockets/Private/BSDIPv6Sockets/SocketSubsystemBSDIPv6.cpp
+++ b/Engine/Source/Runtime/Sockets/Private/BSDIPv6Sockets/SocketSubsystemBSDIPv6.cpp
@@ -47,42 +47,37 @@ void FSocketSubsystemBSDIPv6::DestroySocket(FSocket* Socket)
 
 ESocketErrors FSocketSubsystemBSDIPv6::GetHostByName(const ANSICHAR* HostName, FInternetAddr& OutAddr)
 {
-#if PLATFORM_HAS_BSD_SOCKET_FEATURE_GETHOSTNAME
-	ESocketErrors ErrorCode = SE_NO_ERROR;
-	// gethostbyname() touches a static object so lock for thread safety
 	FScopeLock ScopeLock(&HostByNameSynch);
-	hostent* HostEnt = gethostbyname(HostName);
-	if (HostEnt != NULL)
+	addrinfo* AddrInfo = NULL;
+
+	// We are only interested in IPv6 addresses.
+	addrinfo HintAddrInfo;
+	FMemory::Memzero(&HintAddrInfo, sizeof(HintAddrInfo));
+	HintAddrInfo.ai_family = AF_INET6;
+
+	int32 ErrorCode = getaddrinfo(HostName, NULL, &HintAddrInfo, &AddrInfo);
+	ESocketErrors SocketError = TranslateErrorCode(ErrorCode);
+	if (SocketError == SE_NO_ERROR)
 	{
-		// Make sure it's a valid type
-		if (HostEnt->h_addrtype == PF_INET)
+		for (; AddrInfo != nullptr; AddrInfo = AddrInfo->ai_next)
 		{
-			// Copy the data before letting go of the lock. This is safe only
-			// for the copy locally. If another thread is reading this while
-			// we are copying they will get munged data. This relies on the
-			// consumer of this class to call the resolved() accessor before
-			// attempting to read this data
-			((FInternetAddrBSDIPv6&)OutAddr).SetIp(*(in6_addr*)(*HostEnt->h_addr_list));
-		}
-		else
-		{
-			ErrorCode = SE_HOST_NOT_FOUND;
+			if (AddrInfo->ai_family == AF_INET6)
+			{
+				sockaddr_in6* IPv6SockAddr = reinterpret_cast<sockaddr_in6*>(AddrInfo->ai_addr);
+				if (IPv6SockAddr != nullptr)
+				{
+					static_cast<FInternetAddrBSDIPv6&>(OutAddr).SetIp(IPv6SockAddr->sin6_addr);
+					return SE_NO_ERROR;
+				}
+			}
 		}
+		return SE_HOST_NOT_FOUND;
 	}
-	else
-	{
-		ErrorCode = GetLastErrorCode();
-	}
-	return ErrorCode;
-#else
-	UE_LOG(LogSockets, Error, TEXT("Platform has no gethostbyname(), but did not override FSocketSubsystem::GetHostByName()"));
-	return SE_NO_RECOVERY;
-#endif
+	return SocketError;
 }
 
 bool FSocketSubsystemBSDIPv6::GetHostName(FString& HostName)
 {
-#if PLATFORM_HAS_BSD_SOCKET_FEATURE_GETHOSTNAME
 	ANSICHAR Buffer[256];
 	bool bRead = gethostname(Buffer,256) == 0;
 	if (bRead == true)
@@ -90,10 +85,6 @@ bool FSocketSubsystemBSDIPv6::GetHostName(FString& HostName)
 		HostName = Buffer;
 	}
 	return bRead;
-#else
-	UE_LOG(LogSockets, Error, TEXT("Platform has no gethostname(), but did not override FSocketSubsystem::GetHostName()"));
-	return false;
-#endif
 }
 
 const TCHAR* FSocketSubsystemBSDIPv6::GetSocketAPIName() const
@@ -175,15 +166,9 @@ ESocketErrors FSocketSubsystemBSDIPv6::TranslateErrorCode(int32 Code)
 #if !PLATFORM_HAS_NO_EPROCLIM
 	case EPROCLIM: return SE_EPROCLIM;
 #endif
-// 	case EDISCON: return SE_EDISCON;
-// 	case SYSNOTREADY: return SE_SYSNOTREADY;
-// 	case VERNOTSUPPORTED: return SE_VERNOTSUPPORTED;
-// 	case NOTINITIALISED: return SE_NOTINITIALISED;
-#if PLATFORM_HAS_BSD_SOCKET_FEATURE_GETHOSTNAME
 	case HOST_NOT_FOUND: return SE_HOST_NOT_FOUND;
 	case TRY_AGAIN: return SE_TRY_AGAIN;
 	case NO_RECOVERY: return SE_NO_RECOVERY;
-#endif
 //	case NO_DATA: return SE_NO_DATA;
 		// case : return SE_UDP_ERR_PORT_UNREACH; //@TODO Find it's replacement
 	}
diff --git a/Engine/Source/Runtime/Sockets/Private/BSDSockets/SocketSubsystemBSD.cpp b/Engine/Source/Runtime/Sockets/Private/BSDSockets/SocketSubsystemBSD.cpp
--- a/Engine/Source/Runtime/Sockets/Private/BSDSockets/SocketSubsystemBSD.cpp
+++ b/Engine/Source/Runtime/Sockets/Private/BSDSockets/SocketSubsystemBSD.cpp
@@ -56,34 +56,36 @@ void FSocketSubsystemBSD::DestroySocket(FSocket* Socket)
 ESocketErrors FSocketSubsystemBSD::GetHostByName(const ANSICHAR* HostName, FInternetAddr& OutAddr)
 {
 #if PLATFORM_HAS_BSD_SOCKET_FEATURE_GETHOSTNAME
-	ESocketErrors ErrorCode = SE_NO_ERROR;
-	// gethostbyname() touches a static object so lock for thread safety
 	FScopeLock ScopeLock(&HostByNameSynch);
-	hostent* HostEnt = gethostbyname(HostName);
-	if (HostEnt != NULL)
+	addrinfo* AddrInfo = NULL;
+
+	// Limit the IP Addresses we get back to just IPv4
+	addrinfo HintAddrInfo;
+	FMemory::Memzero(&HintAddrInfo, sizeof(HintAddrInfo));
+	HintAddrInfo.ai_family = AF_INET;
+
+	int32 ErrorCode = getaddrinfo(HostName, NULL, &HintAddrInfo, &AddrInfo);
+	ESocketErrors SocketError = TranslateErrorCode(ErrorCode);
+	if (SocketError == SE_NO_ERROR)
 	{
-		// Make sure it's a valid type
-		if (HostEnt->h_addrtype == PF_INET)
-		{
-			// Copy the data before letting go of the lock. This is safe only
-			// for the copy locally. If another thread is reading this while
-			// we are copying they will get munged data. This relies on the
-			// consumer of this class to call the resolved() accessor before
-			// attempting to read this data
-			((FInternetAddrBSD&)OutAddr).SetIp(*(in_addr*)(*HostEnt->h_addr_list));
-		}
-		else
+		for (; AddrInfo != nullptr; AddrInfo = AddrInfo->ai_next)
 		{
-			ErrorCode = SE_HOST_NOT_FOUND;
+			if (AddrInfo->ai_family == AF_INET)
+			{
+				sockaddr_in* IPv4SockAddr = reinterpret_cast<sockaddr_in*>(AddrInfo->ai_addr);
+				if (IPv4SockAddr != nullptr)
+				{
+					uint32 HostIP = ntohl(IPv4SockAddr->sin_addr.s_addr);
+					static_cast<FInternetAddrBSD&>(OutAddr).SetIp(HostIP);
+					return SE_NO_ERROR;
+				}
+			}
 		}
+		return SE_HOST_NOT_FOUND;
 	}
-	else
-	{
-		ErrorCode = GetLastErrorCode();
-	}
-	return ErrorCode;
+	return SocketError;
 #else
-	UE_LOG(LogSockets, Error, TEXT("Platform has no gethostbyname(), but did not override FSocketSubsystem::GetHostByName()"));
+	UE_LOG(LogSockets, Error, TEXT("Platform has no getaddrinfo(), but did not override FSocketSubsystem::GetHostByName()"));
 	return SE_NO_RECOVERY;
 #endif
 }
@@ -209,6 +211,7 @@ ESocketErrors FSocketSubsystemBSD::TranslateErrorCode(int32 Code)
 	case TRY_AGAIN: return SE_TRY_AGAIN;
 	case NO_RECOVERY: return SE_NO_RECOVERY;
 #endif
+
 //	case NO_DATA: return SE_NO_DATA;
 		// case : return SE_UDP_ERR_PORT_UNREACH; //@TODO Find it's replacement
 	}
diff --git a/Engine/Source/Runtime/Sockets/Private/SocketSubsystem.cpp b/Engine/Source/Runtime/Sockets/Private/SocketSubsystem.cpp
--- a/Engine/Source/Runtime/Sockets/Private/SocketSubsystem.cpp
+++ b/Engine/Source/Runtime/Sockets/Private/SocketSubsystem.cpp
@@ -321,7 +321,7 @@ TSharedRef<FInternetAddr> ISocketSubsystem::GetLocalHostAddr(FOutputDevice& Out,
 		}
 		else
 		{
-			Out.Logf(TEXT("gethostbyname failed (%s)"), GetSocketError());
+			Out.Logf(TEXT("GetHostByName failed (%s)"), GetSocketError());
 		}
 	}
 
