Sockets: Code and documentation cleanup pass

diff --git a/Engine/Source/Runtime/Sockets/Private/BSDSockets/SocketSubsystemBSD.cpp b/Engine/Source/Runtime/Sockets/Private/BSDSockets/SocketSubsystemBSD.cpp
--- a/Engine/Source/Runtime/Sockets/Private/BSDSockets/SocketSubsystemBSD.cpp
+++ b/Engine/Source/Runtime/Sockets/Private/BSDSockets/SocketSubsystemBSD.cpp
@@ -8,28 +8,32 @@
 #include "SocketsBSD.h"
 
 
-class FSocketBSD* FSocketSubsystemBSD::InternalBSDSocketFactory(SOCKET Socket, ESocketType SocketType, const FString& SocketDescription)
+FSocketBSD* FSocketSubsystemBSD::InternalBSDSocketFactory(SOCKET Socket, ESocketType SocketType, const FString& SocketDescription)
 {
 	// return a new socket object
 	return new FSocketBSD(Socket, SocketType, SocketDescription, this);
 }
 
+
 FSocket* FSocketSubsystemBSD::CreateSocket(const FName& SocketType, const FString& SocketDescription, bool bForceUDP)
 {
 	SOCKET Socket = INVALID_SOCKET;
-	FSocket* NewSocket = NULL;
+	FSocket* NewSocket = nullptr;
+
 	switch (SocketType.GetComparisonIndex())
 	{
 	case NAME_DGram:
 		// Creates a data gram (UDP) socket
 		Socket = socket(AF_INET, SOCK_DGRAM, IPPROTO_UDP);
-		NewSocket = (Socket != INVALID_SOCKET) ? InternalBSDSocketFactory(Socket, SOCKTYPE_Datagram, SocketDescription) : NULL;
+		NewSocket = (Socket != INVALID_SOCKET) ? InternalBSDSocketFactory(Socket, SOCKTYPE_Datagram, SocketDescription) : nullptr;
 		break;
+
 	case NAME_Stream:
 		// Creates a stream (TCP) socket
 		Socket = socket(AF_INET, SOCK_STREAM, IPPROTO_TCP);
-		NewSocket = (Socket != INVALID_SOCKET) ? InternalBSDSocketFactory(Socket, SOCKTYPE_Streaming, SocketDescription) : NULL;
+		NewSocket = (Socket != INVALID_SOCKET) ? InternalBSDSocketFactory(Socket, SOCKTYPE_Streaming, SocketDescription) : nullptr;
 		break;
+
 	default:
 		break;
 	}
@@ -42,11 +46,7 @@ FSocket* FSocketSubsystemBSD::CreateSocket(const FName& SocketType, const FStrin
 	return NewSocket;
 }
 
-/**
- * Cleans up a socket class
- *
- * @param Socket the socket object to destroy
- */
+
 void FSocketSubsystemBSD::DestroySocket(FSocket* Socket)
 {
 	delete Socket;
@@ -57,14 +57,14 @@ ESocketErrors FSocketSubsystemBSD::GetHostByName(const ANSICHAR* HostName, FInte
 {
 #if PLATFORM_HAS_BSD_SOCKET_FEATURE_GETHOSTNAME
 	FScopeLock ScopeLock(&HostByNameSynch);
-	addrinfo* AddrInfo = NULL;
+	addrinfo* AddrInfo = nullptr;
 
 	// Limit the IP Addresses we get back to just IPv4
 	addrinfo HintAddrInfo;
 	FMemory::Memzero(&HintAddrInfo, sizeof(HintAddrInfo));
 	HintAddrInfo.ai_family = AF_INET;
 
-	int32 ErrorCode = getaddrinfo(HostName, NULL, &HintAddrInfo, &AddrInfo);
+	int32 ErrorCode = getaddrinfo(HostName, nullptr, &HintAddrInfo, &AddrInfo);
 	ESocketErrors SocketError = TranslateErrorCode(ErrorCode);
 	if (SocketError == SE_NO_ERROR)
 	{
@@ -90,13 +90,7 @@ ESocketErrors FSocketSubsystemBSD::GetHostByName(const ANSICHAR* HostName, FInte
 #endif
 }
 
-/**
- * Determines the name of the local machine
- *
- * @param HostName the string that receives the data
- *
- * @return true if successful, false otherwise
- */
+
 bool FSocketSubsystemBSD::GetHostName(FString& HostName)
 {
 #if PLATFORM_HAS_BSD_SOCKET_FEATURE_GETHOSTNAME
@@ -114,20 +108,12 @@ bool FSocketSubsystemBSD::GetHostName(FString& HostName)
 }
 
 
-/**
- *	Get the name of the socket subsystem
- * @return a string naming this subsystem
- */
 const TCHAR* FSocketSubsystemBSD::GetSocketAPIName() const
 {
 	return TEXT("BSD");
 }
 
-/**
- *	Create a proper FInternetAddr representation
- * @param Address host address
- * @param Port host port
- */
+
 TSharedRef<FInternetAddr> FSocketSubsystemBSD::CreateInternetAddr(uint32 Address, uint32 Port)
 {
 	TSharedRef<FInternetAddr> Result = MakeShareable(new FInternetAddrBSD);
@@ -136,6 +122,7 @@ TSharedRef<FInternetAddr> FSocketSubsystemBSD::CreateInternetAddr(uint32 Address
 	return Result;
 }
 
+
 ESocketErrors FSocketSubsystemBSD::GetLastErrorCode()
 {
 	return TranslateErrorCode(errno);
@@ -222,4 +209,5 @@ ESocketErrors FSocketSubsystemBSD::TranslateErrorCode(int32 Code)
 	return SE_NO_ERROR;
 }
 
+
 #endif	//PLATFORM_HAS_BSD_SOCKETS
diff --git a/Engine/Source/Runtime/Sockets/Private/BSDSockets/SocketSubsystemBSD.h b/Engine/Source/Runtime/Sockets/Private/BSDSockets/SocketSubsystemBSD.h
--- a/Engine/Source/Runtime/Sockets/Private/BSDSockets/SocketSubsystemBSD.h
+++ b/Engine/Source/Runtime/Sockets/Private/BSDSockets/SocketSubsystemBSD.h
@@ -5,6 +5,7 @@
 #include "SocketSubsystem.h"
 #include "SocketSubsystemPackage.h"
 
+
 #if PLATFORM_HAS_BSD_SOCKETS
 
 #include "SocketSubsystemBSDPrivate.h"
@@ -19,19 +20,14 @@ class FSocketSubsystemBSD
 {
 public:
 
-	// Begin ISocketSubsystem interface
+	// ISocketSubsystem interface
 
 	virtual TSharedRef<FInternetAddr> CreateInternetAddr( uint32 Address = 0, uint32 Port = 0 ) override;
-
 	virtual class FSocket* CreateSocket( const FName& SocketType, const FString& SocketDescription, bool bForceUDP = false ) override;
-
 	virtual void DestroySocket( class FSocket* Socket ) override;
-
 	virtual ESocketErrors GetHostByName( const ANSICHAR* HostName, FInternetAddr& OutAddr ) override;
-
 	virtual bool GetHostName( FString& HostName ) override;
-
-	virtual ESocketErrors GetLastErrorCode( ) override;
+	virtual ESocketErrors GetLastErrorCode() override;
 
 	virtual bool GetLocalAdapterAddresses( TArray<TSharedPtr<FInternetAddr> >& OutAdresses ) override
 	{
@@ -42,23 +38,20 @@ public:
 		return true;
 	}
 
-	virtual const TCHAR* GetSocketAPIName( ) const override;
+	virtual const TCHAR* GetSocketAPIName() const override;
 
-	virtual bool RequiresChatDataBeSeparate( ) override
+	virtual bool RequiresChatDataBeSeparate() override
 	{
 		return false;
 	}
 
-	virtual bool RequiresEncryptedPackets( ) override
+	virtual bool RequiresEncryptedPackets() override
 	{
 		return false;
 	}
 
 	virtual ESocketErrors TranslateErrorCode( int32 Code ) override;
 
-	// End ISocketSubsystem interface
-
-
 protected:
 
 	/**
@@ -74,4 +67,6 @@ private:
 	// Used to prevent multiple threads accessing the shared data.
 	FCriticalSection HostByNameSynch;
 };
-#endif
\ No newline at end of file
+
+
+#endif
diff --git a/Engine/Source/Runtime/Sockets/Private/BSDSockets/SocketsBSD.cpp b/Engine/Source/Runtime/Sockets/Private/BSDSockets/SocketsBSD.cpp
--- a/Engine/Source/Runtime/Sockets/Private/BSDSockets/SocketsBSD.cpp
+++ b/Engine/Source/Runtime/Sockets/Private/BSDSockets/SocketsBSD.cpp
@@ -9,10 +9,11 @@
 //#include "Net/NetworkProfiler.h"
 
 #if PLATFORM_HTML5 
-typedef unsigned long   u_long;
+	typedef unsigned long u_long;
 #endif 
 
-/* FSocket interface
+
+/* FSocket overrides
  *****************************************************************************/
 
 bool FSocketBSD::Close(void)
@@ -50,44 +51,6 @@ bool FSocketBSD::Listen(int32 MaxBacklog)
 	return listen(Socket, MaxBacklog) == 0;
 }
 
-ESocketInternalState::Return FSocketBSD::HasState(ESocketInternalState::Param State, FTimespan WaitTime)
-{
-#if PLATFORM_HAS_BSD_SOCKET_FEATURE_SELECT
-	// convert WaitTime to a timeval
-	timeval Time;
-	Time.tv_sec = (int32)WaitTime.GetTotalSeconds();
-	Time.tv_usec = WaitTime.GetMilliseconds() * 1000;
-
-	fd_set SocketSet;
-
-	// Set up the socket sets we are interested in (just this one)
-	FD_ZERO(&SocketSet);
-	FD_SET(Socket, &SocketSet);
-
-	// Check the status of the state
-	int32 SelectStatus = 0;
-	switch (State)
-	{
-		case ESocketInternalState::CanRead:
-			SelectStatus = select(Socket + 1, &SocketSet, NULL, NULL, &Time);
-			break;
-		case ESocketInternalState::CanWrite:
-			SelectStatus = select(Socket + 1, NULL, &SocketSet, NULL, &Time);
-			break;
-		case ESocketInternalState::HasError:
-			SelectStatus = select(Socket + 1, NULL, NULL, &SocketSet, &Time);
-			break;
-	}
-
-	// if the select returns a positive number, the socket had the state, 0 means didn't have it, and negative is API error condition (not socket's error state)
-	return SelectStatus > 0  ? ESocketInternalState::Yes : 
-		   SelectStatus == 0 ? ESocketInternalState::No : 
-		   ESocketInternalState::EncounteredError;
-#else
-	UE_LOG(LogSockets, Fatal, TEXT("This platform doesn't support select(), but FSocketBSD::HasState was not overridden"));
-	return ESocketInternalState::EncounteredError;
-#endif
-}
 
 bool FSocketBSD::HasPendingConnection(bool& bHasPendingConnection)
 {
@@ -95,14 +58,14 @@ bool FSocketBSD::HasPendingConnection(bool& bHasPendingConnection)
 	bHasPendingConnection = false;
 
 	// make sure socket has no error state
-	if (HasState(ESocketInternalState::HasError) == ESocketInternalState::No)
+	if (HasState(ESocketBSDParam::HasError) == ESocketBSDReturn::No)
 	{
 		// get the read state
-		ESocketInternalState::Return State = HasState(ESocketInternalState::CanRead);
+		ESocketBSDReturn State = HasState(ESocketBSDParam::CanRead);
 		
 		// turn the result into the outputs
-		bHasSucceeded = State != ESocketInternalState::EncounteredError;
-		bHasPendingConnection = State == ESocketInternalState::Yes;
+		bHasSucceeded = State != ESocketBSDReturn::EncounteredError;
+		bHasPendingConnection = State == ESocketBSDReturn::Yes;
 	}
 
 	return bHasSucceeded;
@@ -114,7 +77,7 @@ bool FSocketBSD::HasPendingData(uint32& PendingDataSize)
 	PendingDataSize = 0;
 
 	// make sure socket has no error state
-	if (HasState(ESocketInternalState::CanRead) == ESocketInternalState::Yes)
+	if (HasState(ESocketBSDParam::CanRead) == ESocketBSDReturn::Yes)
 	{
 #if PLATFORM_HAS_BSD_SOCKET_FEATURE_IOCTL
 		// See if there is any pending data on the read socket
@@ -131,7 +94,7 @@ bool FSocketBSD::HasPendingData(uint32& PendingDataSize)
 
 FSocket* FSocketBSD::Accept(const FString& SocketDescription)
 {
-	SOCKET NewSocket = accept(Socket,NULL,NULL);
+	SOCKET NewSocket = accept(Socket, NULL, NULL);
 
 	if (NewSocket != INVALID_SOCKET)
 	{
@@ -234,7 +197,7 @@ bool FSocketBSD::Wait(ESocketWaitConditions::Type Condition, FTimespan WaitTime)
 {
 	if ((Condition == ESocketWaitConditions::WaitForRead) || (Condition == ESocketWaitConditions::WaitForReadOrWrite))
 	{
-		if (HasState(ESocketInternalState::CanRead, WaitTime) == ESocketInternalState::Yes)
+		if (HasState(ESocketBSDParam::CanRead, WaitTime) == ESocketBSDReturn::Yes)
 		{
 			return true;
 		}
@@ -242,7 +205,7 @@ bool FSocketBSD::Wait(ESocketWaitConditions::Type Condition, FTimespan WaitTime)
 
 	if ((Condition == ESocketWaitConditions::WaitForWrite) || (Condition == ESocketWaitConditions::WaitForReadOrWrite))
 	{
-		if (HasState(ESocketInternalState::CanWrite, WaitTime) == ESocketInternalState::Yes)
+		if (HasState(ESocketBSDParam::CanWrite, WaitTime) == ESocketBSDReturn::Yes)
 		{
 			return true;
 		}
@@ -257,21 +220,21 @@ ESocketConnectionState FSocketBSD::GetConnectionState(void)
 	ESocketConnectionState CurrentState = SCS_ConnectionError;
 
 	// look for an existing error
-	if (HasState(ESocketInternalState::HasError) == ESocketInternalState::No)
+	if (HasState(ESocketBSDParam::HasError) == ESocketBSDReturn::No)
 	{
 		if (FDateTime::UtcNow() - LastActivityTime > FTimespan::FromSeconds(5))
 		{
 			// get the write state
-			ESocketInternalState::Return WriteState = HasState(ESocketInternalState::CanWrite, FTimespan::FromMilliseconds(1));
-			ESocketInternalState::Return ReadState = HasState(ESocketInternalState::CanRead, FTimespan::FromMilliseconds(1));
+			ESocketBSDReturn WriteState = HasState(ESocketBSDParam::CanWrite, FTimespan::FromMilliseconds(1));
+			ESocketBSDReturn ReadState = HasState(ESocketBSDParam::CanRead, FTimespan::FromMilliseconds(1));
 		
 			// translate yes or no (error is already set)
-			if (WriteState == ESocketInternalState::Yes || ReadState == ESocketInternalState::Yes)
+			if (WriteState == ESocketBSDReturn::Yes || ReadState == ESocketBSDReturn::Yes)
 			{
 				CurrentState = SCS_Connected;
 				LastActivityTime = FDateTime::UtcNow();
 			}
-			else if (WriteState == ESocketInternalState::No && ReadState == ESocketInternalState::No)
+			else if (WriteState == ESocketBSDReturn::No && ReadState == ESocketBSDReturn::No)
 			{
 				CurrentState = SCS_NotConnected;
 			}
@@ -434,4 +397,50 @@ int32 FSocketBSD::GetPortNo(void)
 	return ntohs(Addr.sin_port);
 }
 
+
+/* FSocketBSD implementation
+*****************************************************************************/
+
+ESocketBSDReturn FSocketBSD::HasState(ESocketBSDParam State, FTimespan WaitTime)
+{
+#if PLATFORM_HAS_BSD_SOCKET_FEATURE_SELECT
+	// convert WaitTime to a timeval
+	timeval Time;
+	Time.tv_sec = (int32)WaitTime.GetTotalSeconds();
+	Time.tv_usec = WaitTime.GetMilliseconds() * 1000;
+
+	fd_set SocketSet;
+
+	// Set up the socket sets we are interested in (just this one)
+	FD_ZERO(&SocketSet);
+	FD_SET(Socket, &SocketSet);
+
+	// Check the status of the state
+	int32 SelectStatus = 0;
+	switch (State)
+	{
+	case ESocketBSDParam::CanRead:
+		SelectStatus = select(Socket + 1, &SocketSet, NULL, NULL, &Time);
+		break;
+
+	case ESocketBSDParam::CanWrite:
+		SelectStatus = select(Socket + 1, NULL, &SocketSet, NULL, &Time);
+		break;
+
+	case ESocketBSDParam::HasError:
+		SelectStatus = select(Socket + 1, NULL, NULL, &SocketSet, &Time);
+		break;
+	}
+
+	// if the select returns a positive number, the socket had the state, 0 means didn't have it, and negative is API error condition (not socket's error state)
+	return SelectStatus > 0 ? ESocketBSDReturn::Yes :
+		SelectStatus == 0 ? ESocketBSDReturn::No :
+		ESocketBSDReturn::EncounteredError;
+#else
+	UE_LOG(LogSockets, Fatal, TEXT("This platform doesn't support select(), but FSocketBSD::HasState was not overridden"));
+	return ESocketBSDReturn::EncounteredError;
+#endif
+}
+
+
 #endif	//PLATFORM_HAS_BSD_SOCKETS
diff --git a/Engine/Source/Runtime/Sockets/Private/BSDSockets/SocketsBSD.h b/Engine/Source/Runtime/Sockets/Private/BSDSockets/SocketsBSD.h
--- a/Engine/Source/Runtime/Sockets/Private/BSDSockets/SocketsBSD.h
+++ b/Engine/Source/Runtime/Sockets/Private/BSDSockets/SocketsBSD.h
@@ -10,20 +10,25 @@
 #include "IPAddress.h"
 
 
-namespace ESocketInternalState
+/**
+ * Enumerates BSD socket state parameters.
+ */
+enum class ESocketBSDParam
 {
-	enum Param
-	{
-		CanRead,
-		CanWrite,
-		HasError,
-	};
-	enum Return
-	{
-		Yes,
-		No,
-		EncounteredError,
-	};
+	CanRead,
+	CanWrite,
+	HasError,
+};
+
+
+/**
+ * Enumerates BSD socket state return values.
+ */
+enum class ESocketBSDReturn
+{
+	Yes,
+	No,
+	EncounteredError,
 };
 
 
@@ -58,6 +63,18 @@ public:
 		Close();
 	}
 
+public:
+
+	/**
+	* Gets the Socket for anyone who knows they have an FSocketBSD.
+	*
+	* @return The native socket.
+	*/
+	SOCKET GetNativeSocket()
+	{
+		return Socket;
+	}
+
 public:
 
 	// FSocket overrides
@@ -90,18 +107,18 @@ public:
 	virtual bool SetReceiveBufferSize(int32 Size,int32& NewSize) override;
 	virtual int32 GetPortNo() override;
 
-	/**
-	 * Return the Socket for anyone who knows they have an FSocketBSD
-	 */
-	SOCKET GetNativeSocket()
-	{
-		return Socket;
-	}
-
 protected:
 
 	/** This is generally select(), but makes it easier for platforms without select to replace it. */
-	virtual ESocketInternalState::Return HasState(ESocketInternalState::Param State, FTimespan WaitTime=FTimespan(0));
+	virtual ESocketBSDReturn HasState(ESocketBSDParam State, FTimespan WaitTime = FTimespan(0));
+
+	/** Updates this socket's time of last activity. */
+	void UpdateActivity()
+	{
+		LastActivityTime = FDateTime::UtcNow();
+	}
+
+private:
 
 	/** Holds the BSD socket object. */
 	SOCKET Socket;
@@ -110,7 +127,7 @@ protected:
 	FDateTime LastActivityTime;
 
 	/** Pointer to the subsystem that created it. */
-	ISocketSubsystem * SocketSubsystem;
+	ISocketSubsystem* SocketSubsystem;
 };
 
 
diff --git a/Engine/Source/Runtime/Sockets/Private/Windows/SocketSubsystemWindows.cpp b/Engine/Source/Runtime/Sockets/Private/Windows/SocketSubsystemWindows.cpp
--- a/Engine/Source/Runtime/Sockets/Private/Windows/SocketSubsystemWindows.cpp
+++ b/Engine/Source/Runtime/Sockets/Private/Windows/SocketSubsystemWindows.cpp
@@ -9,45 +9,40 @@
 #include "HideWindowsPlatformTypes.h"
 
 
-FSocketSubsystemWindows* FSocketSubsystemWindows::SocketSingleton = NULL;
+FSocketSubsystemWindows* FSocketSubsystemWindows::SocketSingleton = nullptr;
+
 
-/**
- * Create the socket subsystem for the given platform service
- */
 FName CreateSocketSubsystem( FSocketSubsystemModule& SocketSubsystemModule )
 {
 	FName SubsystemName(TEXT("WINDOWS"));
+
 	// Create and register our singleton factory with the main online subsystem for easy access
 	FSocketSubsystemWindows* SocketSubsystem = FSocketSubsystemWindows::Create();
 	FString Error;
+
 	if (SocketSubsystem->Init(Error))
 	{
 		SocketSubsystemModule.RegisterSocketSubsystem(SubsystemName, SocketSubsystem);
+
 		return SubsystemName;
 	}
-	else
-	{
-		FSocketSubsystemWindows::Destroy();
-		return NAME_None;
-	}
+
+	FSocketSubsystemWindows::Destroy();
+	
+	return NAME_None;
 }
 
-/**
- * Tear down the socket subsystem for the given platform service
- */
+
 void DestroySocketSubsystem( FSocketSubsystemModule& SocketSubsystemModule )
 {
 	SocketSubsystemModule.UnregisterSocketSubsystem(FName(TEXT("WINDOWS")));
 	FSocketSubsystemWindows::Destroy();
 }
 
-/** 
- * Singleton interface for the Windows socket subsystem 
- * @return the only instance of the Windows socket subsystem
- */
+
 FSocketSubsystemWindows* FSocketSubsystemWindows::Create()
 {
-	if (SocketSingleton == NULL)
+	if (SocketSingleton == nullptr)
 	{
 		SocketSingleton = new FSocketSubsystemWindows();
 	}
@@ -55,63 +50,55 @@ FSocketSubsystemWindows* FSocketSubsystemWindows::Create()
 	return SocketSingleton;
 }
 
-/** 
- * Destroy the singleton Windows socket subsystem
- */
+
 void FSocketSubsystemWindows::Destroy()
 {
-	if (SocketSingleton != NULL)
+	if (SocketSingleton != nullptr)
 	{
 		SocketSingleton->Shutdown();
 		delete SocketSingleton;
-		SocketSingleton = NULL;
+		SocketSingleton = nullptr;
 	}
 }
 
-/**
- * Does Windows platform initialization of the sockets library
- *
- * @param Error a string that is filled with error information
- *
- * @return true if initialized ok, false otherwise
- */
+
 bool FSocketSubsystemWindows::Init(FString& Error)
 {
 	bool bSuccess = false;
+
 	if (bTriedToInit == false)
 	{
 		bTriedToInit = true;
 		WSADATA WSAData;
-		// Init WSA
-		int32 Code = WSAStartup(0x0101,&WSAData);
+
+		// initialize WSA
+		int32 Code = WSAStartup(0x0101, &WSAData);
+
 		if (Code == 0)
 		{
 			bSuccess = true;
 			UE_LOG(LogInit, Log, TEXT("%s: version %i.%i (%i.%i), MaxSocks=%i, MaxUdp=%i"),
 				GetSocketAPIName(),
-				WSAData.wVersion >> 8,WSAData.wVersion & 0xFF,
-				WSAData.wHighVersion >> 8,WSAData.wHighVersion & 0xFF,
-				WSAData.iMaxSockets,WSAData.iMaxUdpDg);
+				WSAData.wVersion >> 8, WSAData.wVersion & 0xFF,
+				WSAData.wHighVersion >> 8, WSAData.wHighVersion & 0xFF,
+				WSAData.iMaxSockets, WSAData.iMaxUdpDg);
 		}
 		else
 		{
 			Error = FString::Printf(TEXT("WSAStartup failed (%s)"), GetSocketError(TranslateErrorCode(Code)));
 		}
 	}
+
 	return bSuccess;
 }
 
-/**
- * Performs Windows specific socket clean up
- */
+
 void FSocketSubsystemWindows::Shutdown(void)
 {
 	WSACleanup();
 }
 
-/**
- * Returns the last error that has happened
- */
+
 ESocketErrors FSocketSubsystemWindows::GetLastErrorCode()
 {
 	return TranslateErrorCode(WSAGetLastError());
@@ -236,23 +223,18 @@ ESocketErrors FSocketSubsystemWindows::TranslateErrorCode(int32 Code)
 
 	UE_LOG(LogSockets, Warning, TEXT("Unhandled socket error! Error Code: %d"), Code);
 	check(0);
+
 	return SE_NO_ERROR;
 }
 
 
-/**
- * @return Whether the machine has a properly configured network device or not
- */
 bool FSocketSubsystemWindows::HasNetworkDevice()
 {
 	return true;
 }
 
-/**
- *	Get the name of the socket subsystem
- * @return a string naming this subsystem
- */
+
 const TCHAR* FSocketSubsystemWindows::GetSocketAPIName() const
 {
 	return TEXT("WinSock");
-}
\ No newline at end of file
+}
diff --git a/Engine/Source/Runtime/Sockets/Private/Windows/SocketSubsystemWindows.h b/Engine/Source/Runtime/Sockets/Private/Windows/SocketSubsystemWindows.h
--- a/Engine/Source/Runtime/Sockets/Private/Windows/SocketSubsystemWindows.h
+++ b/Engine/Source/Runtime/Sockets/Private/Windows/SocketSubsystemWindows.h
@@ -9,63 +9,48 @@
 /**
  * Windows specific socket subsystem implementation
  */
-class FSocketSubsystemWindows : public FSocketSubsystemBSD
+class FSocketSubsystemWindows
+	: public FSocketSubsystemBSD
 {
 public:
 
-	/**
-	 * Default constructor.
-	 */
-	FSocketSubsystemWindows( ) :
+	/** Default constructor. */
+	FSocketSubsystemWindows() :
 		bTriedToInit(false)
 	{ }
 
-	/**
-	 * Virtual destructor.
-	 */
-	virtual ~FSocketSubsystemWindows( ) { }
-
+	/** Virtual destructor. */
+	virtual ~FSocketSubsystemWindows() { }
 
 public:
 
-	// Begin ISocketSubsystem interface
-
-	virtual bool HasNetworkDevice( ) override;
-
-	virtual ESocketErrors GetLastErrorCode( ) override;
+	// FSocketSubsystemBSD overrides
 
+	virtual bool HasNetworkDevice() override;
+	virtual ESocketErrors GetLastErrorCode() override;
 	virtual bool GetLocalAdapterAddresses( TArray<TSharedPtr<FInternetAddr> >& OutAdresses ) override;
-
-	virtual const TCHAR* GetSocketAPIName( ) const override;
-
+	virtual const TCHAR* GetSocketAPIName() const override;
 	virtual bool Init( FString& Error ) override;
-
-	virtual void Shutdown( ) override;
-
+	virtual void Shutdown() override;
 	virtual ESocketErrors TranslateErrorCode( int32 Code ) override;
 
-	// End ISocketSubsystem interface
-
-
 PACKAGE_SCOPE:
 
 	/** 
-	 * Singleton interface for this subsystem 
-	 * @return the only instance of this subsystem
-	 */
-	static FSocketSubsystemWindows* Create( );
-
-	/**
-	 * Performs Windows specific socket clean up
+	 * Singleton interface for this subsystem.
+	 *
+	 * @return the only instance of this subsystem.
 	 */
-	static void Destroy( );
+	static FSocketSubsystemWindows* Create();
 
+	/** Performs Windows specific socket clean up. */
+	static void Destroy();
 
 protected:
 
-	// Holds a flag indicating whether Init() has been called before or not.
+	/** Holds a flag indicating whether Init() has been called before or not. */
 	bool bTriedToInit;
 
-	// Holds the single instantiation of this subsystem.
+	/** Holds the single instantiation of this subsystem. */
 	static FSocketSubsystemWindows* SocketSingleton;
 };
\ No newline at end of file
