Sockets: Disabling socket handle inheritance on Windows

Windows enables inheritance for socket inheritance by default, which causes open listener sockets to be passed to client processes. This causes problems in certain scenarios, for example when restarting the Editor to load a new project, or when interfacing with the Android SDK (adb.exe).  Child processes or new instances of the Editor are then not able to bind to the listener ports.

diff --git a/Engine/Source/Runtime/Sockets/Private/Windows/SocketSubsystemWindows.cpp b/Engine/Source/Runtime/Sockets/Private/Windows/SocketSubsystemWindows.cpp
--- a/Engine/Source/Runtime/Sockets/Private/Windows/SocketSubsystemWindows.cpp
+++ b/Engine/Source/Runtime/Sockets/Private/Windows/SocketSubsystemWindows.cpp
@@ -40,6 +40,9 @@ void DestroySocketSubsystem( FSocketSubsystemModule& SocketSubsystemModule )
 }
 
 
+/* FSocketSubsystemWindows interface
+*****************************************************************************/
+
 FSocketSubsystemWindows* FSocketSubsystemWindows::Create()
 {
 	if (SocketSingleton == nullptr)
@@ -62,6 +65,45 @@ void FSocketSubsystemWindows::Destroy()
 }
 
 
+/* FSocketSubsystemBSD overrides
+*****************************************************************************/
+
+FSocket* FSocketSubsystemWindows::CreateSocket(const FName& SocketType, const FString& SocketDescription, bool bForceUDP)
+{
+	SOCKET Socket = INVALID_SOCKET;
+	FSocketBSD* NewSocket = nullptr;
+
+	switch (SocketType.GetComparisonIndex())
+	{
+	case NAME_DGram:
+		// Creates a data gram (UDP) socket
+		Socket = socket(AF_INET, SOCK_DGRAM, IPPROTO_UDP);
+		NewSocket = (Socket != INVALID_SOCKET) ? InternalBSDSocketFactory(Socket, SOCKTYPE_Datagram, SocketDescription) : nullptr;
+		break;
+
+	case NAME_Stream:
+		// Creates a stream (TCP) socket
+		Socket = socket(AF_INET, SOCK_STREAM, IPPROTO_TCP);
+		NewSocket = (Socket != INVALID_SOCKET) ? InternalBSDSocketFactory(Socket, SOCKTYPE_Streaming, SocketDescription) : nullptr;
+		break;
+
+	default:
+		break;
+	}
+
+	if (NewSocket != nullptr)
+	{
+		::SetHandleInformation((HANDLE)NewSocket->GetNativeSocket(), HANDLE_FLAG_INHERIT, 0);
+	}
+	else
+	{
+		UE_LOG(LogSockets, Warning, TEXT("Failed to create socket %s [%s]"), *SocketType.ToString(), *SocketDescription);
+	}
+
+	return NewSocket;
+}
+
+
 bool FSocketSubsystemWindows::Init(FString& Error)
 {
 	bool bSuccess = false;
diff --git a/Engine/Source/Runtime/Sockets/Private/Windows/SocketSubsystemWindows.h b/Engine/Source/Runtime/Sockets/Private/Windows/SocketSubsystemWindows.h
--- a/Engine/Source/Runtime/Sockets/Private/Windows/SocketSubsystemWindows.h
+++ b/Engine/Source/Runtime/Sockets/Private/Windows/SocketSubsystemWindows.h
@@ -4,10 +4,12 @@
 
 #include "SocketSubsystem.h"
 #include "BSDSockets/SocketSubsystemBSD.h"
+#include "BSDSockets/SocketsBSD.h"
 #include "SocketSubsystemPackage.h"
 
+
 /**
- * Windows specific socket subsystem implementation
+ * Windows specific socket subsystem implementation.
  */
 class FSocketSubsystemWindows
 	: public FSocketSubsystemBSD
@@ -26,6 +28,7 @@ public:
 
 	// FSocketSubsystemBSD overrides
 
+	virtual class FSocket* CreateSocket(const FName& SocketType, const FString& SocketDescription, bool bForceUDP = false) override;
 	virtual bool HasNetworkDevice() override;
 	virtual ESocketErrors GetLastErrorCode() override;
 	virtual bool GetLocalAdapterAddresses( TArray<TSharedPtr<FInternetAddr> >& OutAdresses ) override;
