#UE4 Fix non-windows bsd sockets, they all access Socket directly so it can't be private

diff --git a/Engine/Source/Runtime/Sockets/Private/BSDSockets/SocketsBSD.h b/Engine/Source/Runtime/Sockets/Private/BSDSockets/SocketsBSD.h
--- a/Engine/Source/Runtime/Sockets/Private/BSDSockets/SocketsBSD.h
+++ b/Engine/Source/Runtime/Sockets/Private/BSDSockets/SocketsBSD.h
@@ -118,8 +118,6 @@ protected:
 		LastActivityTime = FDateTime::UtcNow();
 	}
 
-private:
-
 	/** Holds the BSD socket object. */
 	SOCKET Socket;
 
