Translate getaddrinfo() in a separate function.

diff --git a/Engine/Source/Runtime/Sockets/Private/BSDIPv6Sockets/SocketSubsystemBSDIPv6.cpp b/Engine/Source/Runtime/Sockets/Private/BSDIPv6Sockets/SocketSubsystemBSDIPv6.cpp
--- a/Engine/Source/Runtime/Sockets/Private/BSDIPv6Sockets/SocketSubsystemBSDIPv6.cpp
+++ b/Engine/Source/Runtime/Sockets/Private/BSDIPv6Sockets/SocketSubsystemBSDIPv6.cpp
@@ -56,7 +56,7 @@ ESocketErrors FSocketSubsystemBSDIPv6::GetHostByName(const ANSICHAR* HostName, F
 	HintAddrInfo.ai_family = AF_INET6;
 
 	int32 ErrorCode = getaddrinfo(HostName, NULL, &HintAddrInfo, &AddrInfo);
-	ESocketErrors SocketError = TranslateErrorCode(ErrorCode);
+	ESocketErrors SocketError = TranslateGAIErrorCode(ErrorCode);
 	if (SocketError == SE_NO_ERROR)
 	{
 		for (; AddrInfo != nullptr; AddrInfo = AddrInfo->ai_next)
diff --git a/Engine/Source/Runtime/Sockets/Private/BSDIPv6Sockets/SocketSubsystemBSDIPv6.h b/Engine/Source/Runtime/Sockets/Private/BSDIPv6Sockets/SocketSubsystemBSDIPv6.h
--- a/Engine/Source/Runtime/Sockets/Private/BSDIPv6Sockets/SocketSubsystemBSDIPv6.h
+++ b/Engine/Source/Runtime/Sockets/Private/BSDIPv6Sockets/SocketSubsystemBSDIPv6.h
@@ -15,7 +15,7 @@
  * Standard BSD specific IPv6 socket subsystem implementation
  */
 class FSocketSubsystemBSDIPv6
-	: public ISocketSubsystem
+	: public FSocketSubsystemBSDCommon
 {
 public:
 
diff --git a/Engine/Source/Runtime/Sockets/Private/BSDSockets/SocketSubsystemBSD.cpp b/Engine/Source/Runtime/Sockets/Private/BSDSockets/SocketSubsystemBSD.cpp
--- a/Engine/Source/Runtime/Sockets/Private/BSDSockets/SocketSubsystemBSD.cpp
+++ b/Engine/Source/Runtime/Sockets/Private/BSDSockets/SocketSubsystemBSD.cpp
@@ -65,7 +65,7 @@ ESocketErrors FSocketSubsystemBSD::GetHostByName(const ANSICHAR* HostName, FInte
 	HintAddrInfo.ai_family = AF_INET;
 
 	int32 ErrorCode = getaddrinfo(HostName, nullptr, &HintAddrInfo, &AddrInfo);
-	ESocketErrors SocketError = TranslateErrorCode(ErrorCode);
+	ESocketErrors SocketError = TranslateGAIErrorCode(ErrorCode);
 	if (SocketError == SE_NO_ERROR)
 	{
 		for (; AddrInfo != nullptr; AddrInfo = AddrInfo->ai_next)
@@ -197,16 +197,6 @@ ESocketErrors FSocketSubsystemBSD::TranslateErrorCode(int32 Code)
 	case HOST_NOT_FOUND: return SE_HOST_NOT_FOUND;
 	case TRY_AGAIN: return SE_TRY_AGAIN;
 	case NO_RECOVERY: return SE_NO_RECOVERY;
-
-	// getaddrinfo() has its own error codes
-	case EAI_AGAIN: return SE_TRY_AGAIN;
-	case EAI_BADFLAGS: return  SE_EINVAL;
-	case EAI_FAIL: return SE_NO_RECOVERY;
-	case EAI_FAMILY: return SE_EAFNOSUPPORT;
-	case EAI_MEMORY: return SE_ENOBUFS;
-	case EAI_NONAME: return SE_HOST_NOT_FOUND;
-	case EAI_SERVICE: return SE_EPFNOSUPPORT;
-	case EAI_SOCKTYPE: return SE_ESOCKTNOSUPPORT;
 #endif
 
 //	case NO_DATA: return SE_NO_DATA;
diff --git a/Engine/Source/Runtime/Sockets/Private/BSDSockets/SocketSubsystemBSD.h b/Engine/Source/Runtime/Sockets/Private/BSDSockets/SocketSubsystemBSD.h
--- a/Engine/Source/Runtime/Sockets/Private/BSDSockets/SocketSubsystemBSD.h
+++ b/Engine/Source/Runtime/Sockets/Private/BSDSockets/SocketSubsystemBSD.h
@@ -16,7 +16,7 @@
  * Standard BSD specific socket subsystem implementation
  */
 class FSocketSubsystemBSD
-	: public ISocketSubsystem
+	: public FSocketSubsystemBSDCommon
 {
 public:
 
diff --git a/Engine/Source/Runtime/Sockets/Private/BSDSockets/SocketSubsystemBSDPrivate.h b/Engine/Source/Runtime/Sockets/Private/BSDSockets/SocketSubsystemBSDPrivate.h
--- a/Engine/Source/Runtime/Sockets/Private/BSDSockets/SocketSubsystemBSDPrivate.h
+++ b/Engine/Source/Runtime/Sockets/Private/BSDSockets/SocketSubsystemBSDPrivate.h
@@ -62,3 +62,39 @@ inline int TranslateFlags(ESocketReceiveFlags::Type Flags)
 }
 
 #include <errno.h>
+
+/**
+ * Standard BSD specific socket subsystem implementation (common to both IPv4 and IPv6)
+ */
+class FSocketSubsystemBSDCommon
+	: public ISocketSubsystem
+{
+protected:
+
+	/**
+	 * Translates return values of getaddrinfo() to socket error enum
+	 */
+	inline ESocketErrors TranslateGAIErrorCode(int32 Code) const
+	{
+#if PLATFORM_HAS_BSD_SOCKET_FEATURE_GETHOSTNAME
+		switch (Code)
+		{
+			// getaddrinfo() has its own error codes
+			case EAI_AGAIN:			return SE_TRY_AGAIN;
+			case EAI_BADFLAGS:		return SE_EINVAL;
+			case EAI_FAIL:			return SE_NO_RECOVERY;
+			case EAI_FAMILY:		return SE_EAFNOSUPPORT;
+			case EAI_MEMORY:		return SE_ENOBUFS;
+			case EAI_NONAME:		return SE_HOST_NOT_FOUND;
+			case EAI_SERVICE:		return SE_EPFNOSUPPORT;
+			case EAI_SOCKTYPE:		return SE_ESOCKTNOSUPPORT;
+			default:
+				UE_LOG(LogSockets, Warning, TEXT("Unhandled getaddrinfo() socket error! Code: %d"), Code);
+				check(0);
+				break;
+		}
+#endif // PLATFORM_HAS_BSD_SOCKET_FEATURE_GETHOSTNAME
+
+		return SE_NO_ERROR;
+	};
+};
