Fixed TranslateGAIErrorCode asserting when there was no error

diff --git a/Engine/Source/Runtime/Sockets/Private/BSDSockets/SocketSubsystemBSDPrivate.h b/Engine/Source/Runtime/Sockets/Private/BSDSockets/SocketSubsystemBSDPrivate.h
--- a/Engine/Source/Runtime/Sockets/Private/BSDSockets/SocketSubsystemBSDPrivate.h
+++ b/Engine/Source/Runtime/Sockets/Private/BSDSockets/SocketSubsystemBSDPrivate.h
@@ -88,6 +88,7 @@ protected:
 			case EAI_NONAME:		return SE_HOST_NOT_FOUND;
 			case EAI_SERVICE:		return SE_EPFNOSUPPORT;
 			case EAI_SOCKTYPE:		return SE_ESOCKTNOSUPPORT;
+			case 0:					break; // 0 means success
 			default:
 				UE_LOG(LogSockets, Warning, TEXT("Unhandled getaddrinfo() socket error! Code: %d"), Code);
 				check(0);
