Linux: return a non-loopback address for local host.

- Ubuntu/Debian/SteamOS have an entry in /etc/hosts that maps hostname to 127.0.1.1
- This breaks the code that tries to advertise local addr over LAN (e.g. OnlineSubsystemNull).
- Resolves UE-5407 (ShooterGame cannot join match)

diff --git a/Engine/Source/Runtime/Sockets/Private/Linux/SocketSubsystemLinux.cpp b/Engine/Source/Runtime/Sockets/Private/Linux/SocketSubsystemLinux.cpp
--- a/Engine/Source/Runtime/Sockets/Private/Linux/SocketSubsystemLinux.cpp
+++ b/Engine/Source/Runtime/Sockets/Private/Linux/SocketSubsystemLinux.cpp
@@ -4,6 +4,8 @@
 #include "SocketSubsystemLinux.h"
 #include "ModuleManager.h"
 
+#include <sys/ioctl.h>
+#include <net/if.h>
 
 FSocketSubsystemLinux* FSocketSubsystemLinux::SocketSingleton = NULL;
 
@@ -83,5 +85,59 @@ void FSocketSubsystemLinux::Shutdown(void)
  */
 bool FSocketSubsystemLinux::HasNetworkDevice()
 {
+	// @TODO: implement
 	return true;
 }
+
+TSharedRef<FInternetAddr> FSocketSubsystemLinux::GetLocalHostAddr(FOutputDevice& Out, bool& bCanBindAll)
+{
+	// get parent address first
+	TSharedRef<FInternetAddr> Addr = FSocketSubsystemBSD::GetLocalHostAddr(Out, bCanBindAll);
+
+	// if the address is not a loopback one (or none), return it
+	uint32 ParentIp = 0;
+	Addr->GetIp(ParentIp); // will return in host order
+	if (ParentIp != 0 && (ParentIp & 0xff000000) != 0x7f000000)
+	{
+		return Addr;
+	}
+	
+	// we need to go deeper...  (see http://unixhelp.ed.ac.uk/CGI/man-cgi?netdevice+7)
+	int TempSocket = socket(PF_INET, SOCK_STREAM, 0);
+	if (TempSocket)
+	{
+		ifreq IfReqs[8];
+		
+		ifconf IfConfig;
+		FMemory::MemZero(IfConfig);
+		IfConfig.ifc_req = IfReqs;
+		IfConfig.ifc_len = sizeof(IfReqs);
+		
+		int Result = ioctl(TempSocket, SIOCGIFCONF, &IfConfig);
+		if (Result == 0)
+		{
+			for (int32 IdxReq = 0; IdxReq < ARRAY_COUNT(IfReqs); ++IdxReq)
+			{
+				// grab the first non-loobpack one which is up
+				int ResultFlags = ioctl(TempSocket, SIOCGIFFLAGS, &IfReqs[IdxReq]);
+				if (ResultFlags == 0 && 
+					(IfReqs[IdxReq].ifr_flags & IFF_UP) && 
+					(IfReqs[IdxReq].ifr_flags & IFF_LOOPBACK) == 0)
+				{
+					int32 NetworkAddr = reinterpret_cast<sockaddr_in *>(&IfReqs[IdxReq].ifr_addr)->sin_addr.s_addr;
+					Addr->SetIp(ntohl(NetworkAddr));
+					break;
+				}
+			}
+		}
+		else
+		{
+			int ErrNo = errno;
+			UE_LOG(LogSockets, Warning, TEXT("ioctl( ,SIOGCIFCONF, ) failed, errno=%d (%s)"), ErrNo, ANSI_TO_TCHAR(strerror(ErrNo)));
+		}
+		
+		close(TempSocket);
+	}
+
+	return Addr;
+}
diff --git a/Engine/Source/Runtime/Sockets/Private/Linux/SocketSubsystemLinux.h b/Engine/Source/Runtime/Sockets/Private/Linux/SocketSubsystemLinux.h
--- a/Engine/Source/Runtime/Sockets/Private/Linux/SocketSubsystemLinux.h
+++ b/Engine/Source/Runtime/Sockets/Private/Linux/SocketSubsystemLinux.h
@@ -45,22 +45,9 @@ public:
 	{
 	}
 
-	/**
-	 * Does Linux platform initialization of the sockets library
-	 *
-	 * @param Error a string that is filled with error information
-	 *
-	 * @return TRUE if initialized ok, FALSE otherwise
-	 */
+	// ISocketSubsystem
 	virtual bool Init(FString& Error) override;
-
-	/**
-	 * Performs platform specific socket clean up
-	 */
 	virtual void Shutdown() override;
-
-	/**
-	 * @return Whether the device has a properly configured network device or not
-	 */
 	virtual bool HasNetworkDevice() override;
+	virtual TSharedRef<FInternetAddr> GetLocalHostAddr(FOutputDevice& Out, bool& bCanBindAll);
 };
