Changed FSocket.GetPeerAddress to return a bool.

diff --git a/Engine/Source/Runtime/Online/OnlineSubsystemSteam/Private/SocketsSteam.cpp b/Engine/Source/Runtime/Online/OnlineSubsystemSteam/Private/SocketsSteam.cpp
--- a/Engine/Source/Runtime/Online/OnlineSubsystemSteam/Private/SocketsSteam.cpp
+++ b/Engine/Source/Runtime/Online/OnlineSubsystemSteam/Private/SocketsSteam.cpp
@@ -228,11 +228,10 @@ void FSocketSteam::GetAddress(FInternetAddr& OutAddr)
 	SteamAddr.SteamChannel = SteamChannel;
 }
 
-void FSocketSteam::GetPeerAddress(FInternetAddr& OutAddr)
+bool FSocketSteam::GetPeerAddress(FInternetAddr& OutAddr)
 {
-	// don't support this
-	FInternetAddrSteam& SteamAddr = (FInternetAddrSteam&)OutAddr;
-	new( &SteamAddr ) FInternetAddrSteam();
+	// don't support this	
+	return false;
 }
 
 /**
diff --git a/Engine/Source/Runtime/Online/OnlineSubsystemSteam/Private/SocketsSteam.h b/Engine/Source/Runtime/Online/OnlineSubsystemSteam/Private/SocketsSteam.h
--- a/Engine/Source/Runtime/Online/OnlineSubsystemSteam/Private/SocketsSteam.h
+++ b/Engine/Source/Runtime/Online/OnlineSubsystemSteam/Private/SocketsSteam.h
@@ -194,7 +194,7 @@ public:
 	 * 
 	 * @param OutAddr address of the peer the socket is connected to
 	 */
-	virtual void GetPeerAddress(FInternetAddr& OutAddr) override;
+	virtual bool GetPeerAddress(FInternetAddr& OutAddr) override;
 
 	/**
 	 * Sets this socket into non-blocking mode
diff --git a/Engine/Source/Runtime/Sockets/Private/BSDIPv6Sockets/SocketsBSDIPv6.cpp b/Engine/Source/Runtime/Sockets/Private/BSDIPv6Sockets/SocketsBSDIPv6.cpp
--- a/Engine/Source/Runtime/Sockets/Private/BSDIPv6Sockets/SocketsBSDIPv6.cpp
+++ b/Engine/Source/Runtime/Sockets/Private/BSDIPv6Sockets/SocketsBSDIPv6.cpp
@@ -270,19 +270,20 @@ void FSocketBSDIPv6::GetAddress(FInternetAddr& OutAddr)
 	}
 }
 
-void FSocketBSDIPv6::GetPeerAddress(FInternetAddr& OutAddr)
+bool FSocketBSDIPv6::GetPeerAddress(FInternetAddr& OutAddr)
 {
 	FInternetAddrBSDIPv6& Addr = (FInternetAddrBSDIPv6&)OutAddr;
 	SOCKLEN Size = sizeof(sockaddr_in6);
 
 	// Figure out what ip/port we are bound to
-	bool bOk = getpeername(Socket, Addr, &Size) == 0;
+	int Result = getpeername(Socket, Addr, &Size);
 
-	if (bOk == false)
+	if (Result != 0)
 	{
 		check(SocketSubsystem);
-		UE_LOG(LogSockets, Error, TEXT("Failed to read address for socket (%s)"), SocketSubsystem->GetSocketError());
+		UE_LOG(LogSockets, Warning, TEXT("Failed to read address for socket (%s) with error %d"), SocketSubsystem->GetSocketError(), Result);
 	}
+	return Result == 0;
 }
 
 
diff --git a/Engine/Source/Runtime/Sockets/Private/BSDIPv6Sockets/SocketsBSDIPv6.h b/Engine/Source/Runtime/Sockets/Private/BSDIPv6Sockets/SocketsBSDIPv6.h
--- a/Engine/Source/Runtime/Sockets/Private/BSDIPv6Sockets/SocketsBSDIPv6.h
+++ b/Engine/Source/Runtime/Sockets/Private/BSDIPv6Sockets/SocketsBSDIPv6.h
@@ -90,7 +90,7 @@ public:
 
 	virtual void GetAddress(FInternetAddr& OutAddr) override;
 
-	virtual void GetPeerAddress(FInternetAddr& OutAddr) override;
+	virtual bool GetPeerAddress(FInternetAddr& OutAddr) override;
 
 	virtual bool SetNonBlocking(bool bIsNonBlocking = true) override;
 
diff --git a/Engine/Source/Runtime/Sockets/Private/BSDSockets/SocketsBSD.cpp b/Engine/Source/Runtime/Sockets/Private/BSDSockets/SocketsBSD.cpp
--- a/Engine/Source/Runtime/Sockets/Private/BSDSockets/SocketsBSD.cpp
+++ b/Engine/Source/Runtime/Sockets/Private/BSDSockets/SocketsBSD.cpp
@@ -265,19 +265,20 @@ void FSocketBSD::GetAddress(FInternetAddr& OutAddr)
 }
 
 
-void FSocketBSD::GetPeerAddress(FInternetAddr& OutAddr)
+bool FSocketBSD::GetPeerAddress(FInternetAddr& OutAddr)
 {
 	FInternetAddrBSD& Addr = (FInternetAddrBSD&)OutAddr;
 	SOCKLEN Size = sizeof(sockaddr_in);
 
 	// Figure out what ip/port we are bound to
-	bool bOk = getpeername(Socket, Addr, &Size) == 0;
+	int Result = getpeername(Socket, Addr, &Size);
 
-	if (bOk == false)
+	if (Result != 0)
 	{
 		check(SocketSubsystem);
-		UE_LOG(LogSockets, Error, TEXT("Failed to read address for socket (%s)"), SocketSubsystem->GetSocketError());
+		UE_LOG(LogSockets, Warning, TEXT("Failed to read address for socket (%s) with error %d"), SocketSubsystem->GetSocketError(), Result);
 	}
+	return Result == 0;
 }
 
 bool FSocketBSD::SetNonBlocking(bool bIsNonBlocking)
diff --git a/Engine/Source/Runtime/Sockets/Private/BSDSockets/SocketsBSD.h b/Engine/Source/Runtime/Sockets/Private/BSDSockets/SocketsBSD.h
--- a/Engine/Source/Runtime/Sockets/Private/BSDSockets/SocketsBSD.h
+++ b/Engine/Source/Runtime/Sockets/Private/BSDSockets/SocketsBSD.h
@@ -94,7 +94,7 @@ public:
 	virtual bool Wait(ESocketWaitConditions::Type Condition, FTimespan WaitTime) override;
 	virtual ESocketConnectionState GetConnectionState() override;
 	virtual void GetAddress(FInternetAddr& OutAddr) override;
-	virtual void GetPeerAddress(FInternetAddr& OutAddr) override;
+	virtual bool GetPeerAddress(FInternetAddr& OutAddr) override;
 	virtual bool SetNonBlocking(bool bIsNonBlocking = true) override;
 	virtual bool SetBroadcast(bool bAllowBroadcast = true) override;
 	virtual bool JoinMulticastGroup(const FInternetAddr& GroupAddress) override;
diff --git a/Engine/Source/Runtime/Sockets/Private/HTML5/Simulator/SocketX.cpp b/Engine/Source/Runtime/Sockets/Private/HTML5/Simulator/SocketX.cpp
--- a/Engine/Source/Runtime/Sockets/Private/HTML5/Simulator/SocketX.cpp
+++ b/Engine/Source/Runtime/Sockets/Private/HTML5/Simulator/SocketX.cpp
@@ -108,6 +108,11 @@ void FSocketX::GetAddress( FInternetAddr& OutAddr )
 	Pimpl->GetAddress(ADDTOPIMPL(OutAddr) );
 }
 
+bool FSocketX::GetPeerAddress( FInternetAddr& OutAddr )
+{
+	Pimpl->GetPeerAddress(ADDTOPIMPL(OutAddr));
+}
+
 bool FSocketX::SetNonBlocking( bool bIsNonBlocking /*= true */ )
 {
 	return Pimpl->SetNonBlocking( bIsNonBlocking ); 
diff --git a/Engine/Source/Runtime/Sockets/Private/HTML5/Simulator/SocketX.h b/Engine/Source/Runtime/Sockets/Private/HTML5/Simulator/SocketX.h
--- a/Engine/Source/Runtime/Sockets/Private/HTML5/Simulator/SocketX.h
+++ b/Engine/Source/Runtime/Sockets/Private/HTML5/Simulator/SocketX.h
@@ -45,6 +45,8 @@ public:
 
 	virtual void GetAddress( FInternetAddr& OutAddr );
 
+	virtual bool GetPeerAddress( FInternetAddr& OutAddr ) override;
+
 	virtual bool SetNonBlocking( bool bIsNonBlocking = true );
 
 	virtual bool SetBroadcast( bool bAllowBroadcast = true );
diff --git a/Engine/Source/Runtime/Sockets/Public/Sockets.h b/Engine/Source/Runtime/Sockets/Public/Sockets.h
--- a/Engine/Source/Runtime/Sockets/Public/Sockets.h
+++ b/Engine/Source/Runtime/Sockets/Public/Sockets.h
@@ -181,8 +181,9 @@ public:
 	 * Reads the address of the peer the socket is connected to
 	 *
 	 * @param OutAddr address of the peer the socket is connected to
+	 * @return return true if the address was retrieved correctly 
 	 */
-	virtual void GetPeerAddress(FInternetAddr& OutAddr) = 0;
+	virtual bool GetPeerAddress(FInternetAddr& OutAddr) = 0;
 
 	/**
 	 * Sets this socket into non-blocking mode
