Merging 2456661

Fix assert and still returning no error if an unexpected error code happens in bsd sockets getaddrinfo()

//depot/UE4-UT/...

to //depot/UE4/...

diff --git a/Engine/Source/Runtime/Sockets/Private/BSDSockets/SocketSubsystemBSDPrivate.h b/Engine/Source/Runtime/Sockets/Private/BSDSockets/SocketSubsystemBSDPrivate.h
--- a/Engine/Source/Runtime/Sockets/Private/BSDSockets/SocketSubsystemBSDPrivate.h
+++ b/Engine/Source/Runtime/Sockets/Private/BSDSockets/SocketSubsystemBSDPrivate.h
@@ -91,8 +91,7 @@ protected:
 			case 0:					break; // 0 means success
 			default:
 				UE_LOG(LogSockets, Warning, TEXT("Unhandled getaddrinfo() socket error! Code: %d"), Code);
-				check(0);
-				break;
+				return SE_EINVAL;
 		}
 #endif // PLATFORM_HAS_BSD_SOCKET_FEATURE_GETHOSTNAME
 
