Fix for beacons not binding to 0.0.0.0.

The "smallest footprint" change, as the logic is a bit convoluted (the code is used by beacons and by OnlineSubsystemNull for slightly different purposes).

- Do not consider EWOULDBLOCK a failure and instead proceed to bind to any address.
- Honor -multihome= command line arg.

diff --git a/Engine/Source/Runtime/Sockets/Private/Linux/SocketSubsystemLinux.cpp b/Engine/Source/Runtime/Sockets/Private/Linux/SocketSubsystemLinux.cpp
--- a/Engine/Source/Runtime/Sockets/Private/Linux/SocketSubsystemLinux.cpp
+++ b/Engine/Source/Runtime/Sockets/Private/Linux/SocketSubsystemLinux.cpp
@@ -94,14 +94,21 @@ TSharedRef<FInternetAddr> FSocketSubsystemLinux::GetLocalHostAddr(FOutputDevice&
 	// get parent address first
 	TSharedRef<FInternetAddr> Addr = FSocketSubsystemBSD::GetLocalHostAddr(Out, bCanBindAll);
 
-	// if the address is not a loopback one (or none), return it
+	// If the address is not a loopback one (or none), return it.
 	uint32 ParentIp = 0;
 	Addr->GetIp(ParentIp); // will return in host order
 	if (ParentIp != 0 && (ParentIp & 0xff000000) != 0x7f000000)
 	{
 		return Addr;
 	}
-	
+
+	// If superclass got the address from command line, honor that override
+	TCHAR Home[256]=TEXT("");
+	if (FParse::Value(FCommandLine::Get(),TEXT("MULTIHOME="),Home,ARRAY_COUNT(Home)))
+	{
+		return Addr;
+	}
+
 	// we need to go deeper...  (see http://unixhelp.ed.ac.uk/CGI/man-cgi?netdevice+7)
 	int TempSocket = socket(PF_INET, SOCK_STREAM, 0);
 	if (TempSocket)
diff --git a/Engine/Source/Runtime/Sockets/Private/SocketSubsystem.cpp b/Engine/Source/Runtime/Sockets/Private/SocketSubsystem.cpp
--- a/Engine/Source/Runtime/Sockets/Private/SocketSubsystem.cpp
+++ b/Engine/Source/Runtime/Sockets/Private/SocketSubsystem.cpp
@@ -308,7 +308,7 @@ TSharedRef<FInternetAddr> ISocketSubsystem::GetLocalHostAddr(FOutputDevice& Out,
 	{
 		// Failing to find the host is not considered an error and we just bind to any address
 		ESocketErrors FindHostResult = GetHostByName(TCHAR_TO_ANSI(*HostName), *HostAddr);
-		if (FindHostResult == SE_NO_ERROR || FindHostResult == SE_HOST_NOT_FOUND)
+		if (FindHostResult == SE_NO_ERROR || FindHostResult == SE_HOST_NOT_FOUND || FindHostResult == SE_EWOULDBLOCK)
 		{
 			if( !FParse::Param(FCommandLine::Get(),TEXT("PRIMARYNET")) )
 			{
