Hardened the beacons fix.

- Do not honor -multihome unless valid.
- More getaddrinfo() errors to ignore + more accurate error reporting.

diff --git a/Engine/Source/Runtime/Sockets/Private/Linux/SocketSubsystemLinux.cpp b/Engine/Source/Runtime/Sockets/Private/Linux/SocketSubsystemLinux.cpp
--- a/Engine/Source/Runtime/Sockets/Private/Linux/SocketSubsystemLinux.cpp
+++ b/Engine/Source/Runtime/Sockets/Private/Linux/SocketSubsystemLinux.cpp
@@ -106,7 +106,13 @@ TSharedRef<FInternetAddr> FSocketSubsystemLinux::GetLocalHostAddr(FOutputDevice&
 	TCHAR Home[256]=TEXT("");
 	if (FParse::Value(FCommandLine::Get(),TEXT("MULTIHOME="),Home,ARRAY_COUNT(Home)))
 	{
-		return Addr;
+		TSharedRef<FInternetAddr> TempAddr = CreateInternetAddr();
+		bool bIsValid = false;
+		TempAddr->SetIp(Home, bIsValid);
+		if (bIsValid)
+		{
+			return Addr;
+		}
 	}
 
 	// we need to go deeper...  (see http://unixhelp.ed.ac.uk/CGI/man-cgi?netdevice+7)
diff --git a/Engine/Source/Runtime/Sockets/Private/SocketSubsystem.cpp b/Engine/Source/Runtime/Sockets/Private/SocketSubsystem.cpp
--- a/Engine/Source/Runtime/Sockets/Private/SocketSubsystem.cpp
+++ b/Engine/Source/Runtime/Sockets/Private/SocketSubsystem.cpp
@@ -308,7 +308,8 @@ TSharedRef<FInternetAddr> ISocketSubsystem::GetLocalHostAddr(FOutputDevice& Out,
 	{
 		// Failing to find the host is not considered an error and we just bind to any address
 		ESocketErrors FindHostResult = GetHostByName(TCHAR_TO_ANSI(*HostName), *HostAddr);
-		if (FindHostResult == SE_NO_ERROR || FindHostResult == SE_HOST_NOT_FOUND || FindHostResult == SE_EWOULDBLOCK)
+		if (FindHostResult == SE_NO_ERROR || FindHostResult == SE_HOST_NOT_FOUND 
+			|| FindHostResult == SE_EWOULDBLOCK || FindHostResult == SE_TRY_AGAIN)
 		{
 			if( !FParse::Param(FCommandLine::Get(),TEXT("PRIMARYNET")) )
 			{
@@ -323,7 +324,7 @@ TSharedRef<FInternetAddr> ISocketSubsystem::GetLocalHostAddr(FOutputDevice& Out,
 		}
 		else
 		{
-			Out.Logf(TEXT("GetHostByName failed (%s)"), GetSocketError());
+			Out.Logf(TEXT("GetHostByName failed (%s)"), GetSocketError(FindHostResult));
 		}
 	}
 
