Add winsock specific error messages to TranslateGAIErrorCode

diff --git a/Engine/Source/Runtime/Sockets/Private/BSDSockets/SocketSubsystemBSDPrivate.h b/Engine/Source/Runtime/Sockets/Private/BSDSockets/SocketSubsystemBSDPrivate.h
--- a/Engine/Source/Runtime/Sockets/Private/BSDSockets/SocketSubsystemBSDPrivate.h
+++ b/Engine/Source/Runtime/Sockets/Private/BSDSockets/SocketSubsystemBSDPrivate.h
@@ -88,6 +88,10 @@ protected:
 			case EAI_NONAME:		return SE_HOST_NOT_FOUND;
 			case EAI_SERVICE:		return SE_EPFNOSUPPORT;
 			case EAI_SOCKTYPE:		return SE_ESOCKTNOSUPPORT;
+#if PLATFORM_HAS_BSD_SOCKET_FEATURE_WINSOCKETS
+			case WSANO_DATA:		return SE_NO_DATA;
+			case WSANOTINITIALISED: return SE_NOTINITIALISED;
+#endif
 			case 0:					break; // 0 means success
 			default:
 				UE_LOG(LogSockets, Warning, TEXT("Unhandled getaddrinfo() socket error! Code: %d"), Code);
