Add a couple more error codes for getaddrinfo

diff --git a/Engine/Source/Runtime/Sockets/Private/BSDSockets/SocketSubsystemBSDPrivate.h b/Engine/Source/Runtime/Sockets/Private/BSDSockets/SocketSubsystemBSDPrivate.h
--- a/Engine/Source/Runtime/Sockets/Private/BSDSockets/SocketSubsystemBSDPrivate.h
+++ b/Engine/Source/Runtime/Sockets/Private/BSDSockets/SocketSubsystemBSDPrivate.h
@@ -91,6 +91,10 @@ protected:
 #if PLATFORM_HAS_BSD_SOCKET_FEATURE_WINSOCKETS
 			case WSANO_DATA:		return SE_NO_DATA;
 			case WSANOTINITIALISED: return SE_NOTINITIALISED;
+#else			
+			case EAI_NODATA:		return SE_NO_DATA;
+			case EAI_ADDRFAMILY:	return SE_ADDRFAMILY;
+			case EAI_SYSTEM:		return SE_SYSTEM;
 #endif
 			case 0:					break; // 0 means success
 			default:
diff --git a/Engine/Source/Runtime/Sockets/Public/SocketTypes.h b/Engine/Source/Runtime/Sockets/Public/SocketTypes.h
--- a/Engine/Source/Runtime/Sockets/Public/SocketTypes.h
+++ b/Engine/Source/Runtime/Sockets/Public/SocketTypes.h
@@ -78,6 +78,8 @@ enum ESocketErrors
 	SE_NO_RECOVERY,
 	SE_NO_DATA,
 	SE_UDP_ERR_PORT_UNREACH,
+	SE_ADDRFAMILY,
+	SE_SYSTEM,
 
 	// this is a special error which means to lookup the most recent error (via GetLastErrorCode())
 	SE_GET_LAST_ERROR_CODE,
