Shadowing fixes in IPv6 socket.

diff --git a/Engine/Source/Runtime/Sockets/Private/BSDIPv6Sockets/IPAddressBSDIPv6.h b/Engine/Source/Runtime/Sockets/Private/BSDIPv6Sockets/IPAddressBSDIPv6.h
--- a/Engine/Source/Runtime/Sockets/Private/BSDIPv6Sockets/IPAddressBSDIPv6.h
+++ b/Engine/Source/Runtime/Sockets/Private/BSDIPv6Sockets/IPAddressBSDIPv6.h
@@ -48,10 +48,10 @@ public:
 		}
 		else
 		{
-			in_addr Addr;
-			Addr.s_addr= htonl(InAddr);
+			in_addr InternetAddr;
+			InternetAddr.s_addr= htonl(InAddr);
 
-			SetIp(Addr);
+			SetIp(InternetAddr);
 		}
 	}
 
@@ -157,13 +157,13 @@ public:
 	{
 		if (IpAddr.ss_family == AF_INET)
 		{
-			const sockaddr_in* Addr = (const sockaddr_in*)&IpAddr;
-			SetIp(Addr->sin_addr);
+			const sockaddr_in* SockAddr = (const sockaddr_in*)&IpAddr;
+			SetIp(SockAddr->sin_addr);
 		}
 		else if (IpAddr.ss_family == AF_INET6)
 		{
-			const sockaddr_in6* Addr = (const sockaddr_in6*)&IpAddr;
-			SetIp(Addr->sin6_addr);
+			const sockaddr_in6* SockAddr = (const sockaddr_in6*)&IpAddr;
+			SetIp(SockAddr->sin6_addr);
 		}
 	}
 
