Copying CL#2956821 from Max.Preussner to Dev-Main:

Sockets: Fixed incorrect socket timeout when value < 1 ms

diff --git a/Engine/Source/Runtime/Sockets/Private/BSDIPv6Sockets/SocketsBSDIPv6.cpp b/Engine/Source/Runtime/Sockets/Private/BSDIPv6Sockets/SocketsBSDIPv6.cpp
--- a/Engine/Source/Runtime/Sockets/Private/BSDIPv6Sockets/SocketsBSDIPv6.cpp
+++ b/Engine/Source/Runtime/Sockets/Private/BSDIPv6Sockets/SocketsBSDIPv6.cpp
@@ -49,7 +49,7 @@ EIPv6SocketInternalState::Return FSocketBSDIPv6::HasState(EIPv6SocketInternalSta
 	// convert WaitTime to a timeval
 	timeval Time;
 	Time.tv_sec = (int32)WaitTime.GetTotalSeconds();
-	Time.tv_usec = WaitTime.GetMicroseconds();
+	Time.tv_usec = WaitTime.GetMilliseconds() * 1000 + WaitTime.GetMicroseconds();
 
 	fd_set SocketSet;
 
diff --git a/Engine/Source/Runtime/Sockets/Private/BSDSockets/SocketsBSD.cpp b/Engine/Source/Runtime/Sockets/Private/BSDSockets/SocketsBSD.cpp
--- a/Engine/Source/Runtime/Sockets/Private/BSDSockets/SocketsBSD.cpp
+++ b/Engine/Source/Runtime/Sockets/Private/BSDSockets/SocketsBSD.cpp
@@ -426,7 +426,7 @@ ESocketBSDReturn FSocketBSD::HasState(ESocketBSDParam State, FTimespan WaitTime)
 	// convert WaitTime to a timeval
 	timeval Time;
 	Time.tv_sec = (int32)WaitTime.GetTotalSeconds();
-	Time.tv_usec = WaitTime.GetMicroseconds();
+	Time.tv_usec = WaitTime.GetMilliseconds() * 1000 + WaitTime.GetMicroseconds();
 
 	fd_set SocketSet;
 
