Copying //UE4/Dev-Sequencer to //UE4/Dev-Main (Source: //UE4/Dev-Sequencer @ 3178529)

diff --git a/Engine/Source/Runtime/Sockets/Public/Sockets.h b/Engine/Source/Runtime/Sockets/Public/Sockets.h
--- a/Engine/Source/Runtime/Sockets/Public/Sockets.h
+++ b/Engine/Source/Runtime/Sockets/Public/Sockets.h
@@ -6,12 +6,14 @@
 #include "IPAddress.h"
 #include "SocketTypes.h"
 
+
 /**
  * This is our abstract base class that hides the platform specific socket implementation
  */
 class SOCKETS_API FSocket
 {
 protected:
+
 	/** Indicates the type of socket this is */
 	const ESocketType SocketType;
 
@@ -19,138 +21,128 @@ protected:
 	FString SocketDescription;
 
 public:
-	/** Default ctor */
+
+	/** Default constructor. */
 	inline FSocket() :
 		SocketType(SOCKTYPE_Unknown),
 		SocketDescription(TEXT(""))
-	{
-	}
+	{ }
 
 	/**
 	 * Specifies the type of socket being created
 	 *
-	 * @param InSocketType the type of socket being created
-	 * @param InSocketDescription the debug description of the socket
+	 * @param InSocketType The type of socket being created
+	 * @param InSocketDescription The debug description of the socket
 	 */
 	inline FSocket(ESocketType InSocketType, const FString& InSocketDescription) :
 		SocketType(InSocketType),
 		SocketDescription(InSocketDescription)
-	{
-	}
+	{ }
 
-	/**
-	 * Virtual destructor
-	 */
+	/** Virtual destructor. */
 	virtual ~FSocket()
-	{
-	}
+	{ }
 
 	/**
 	 * Closes the socket
 	 *
-	 * @return true if it closes without errors, false otherwise
+	 * @return true if it closes without errors, false otherwise.
 	 */
 	virtual bool Close() = 0;
 
 	/**
-	 * Binds a socket to a network byte ordered address
-	 *
-	 * @param Addr the address to bind to
+	 * Binds a socket to a network byte ordered address.
 	 *
-	 * @return true if successful, false otherwise
+	 * @param Addr The address to bind to.
+	 * @return true if successful, false otherwise.
 	 */
 	virtual bool Bind(const FInternetAddr& Addr) = 0;
 
 	/**
-	 * Connects a socket to a network byte ordered address
+	 * Connects a socket to a network byte ordered address.
 	 *
-	 * @param Addr the address to connect to
-	 *
-	 * @return true if successful, false otherwise
+	 * @param Addr The address to connect to.
+	 * @return true if successful, false otherwise.
 	 */
 	virtual bool Connect(const FInternetAddr& Addr) = 0;
 
 	/**
-	 * Places the socket into a state to listen for incoming connections
-	 *
-	 * @param MaxBacklog the number of connections to queue before refusing them
+	 * Places the socket into a state to listen for incoming connections.
 	 *
-	 * @return true if successful, false otherwise
+	 * @param MaxBacklog The number of connections to queue before refusing them.
+	 * @return true if successful, false otherwise.
 	 */
 	virtual bool Listen(int32 MaxBacklog) = 0;
 
 	/**
-	 * Queries the socket to determine if there is a pending connection
+	 * Queries the socket to determine if there is a pending connection.
 	 *
-	 * @param bHasPendingConnection out parameter indicating whether a connection is pending or not
-	 *
-	 * @return true if successful, false otherwise
+	 * @param bHasPendingConnection Will indicate whether a connection is pending or not.
+	 * @return true if successful, false otherwise.
 	 */
 	virtual bool HasPendingConnection(bool& bHasPendingConnection) = 0;
 
 	/**
-	* Queries the socket to determine if there is pending data on the queue
-	*
-	* @param PendingDataSize out parameter indicating how much data is on the pipe for a single recv call
+	* Queries the socket to determine if there is pending data on the queue.
 	*
-	* @return true if the socket has data, false otherwise
+	* @param PendingDataSize Will indicate how much data is on the pipe for a single recv call.
+	* @return true if the socket has data, false otherwise.
 	*/
 	virtual bool HasPendingData(uint32& PendingDataSize) = 0;
 
 	/**
-	 * Accepts a connection that is pending
+	 * Accepts a connection that is pending.
 	 *
-	 * @param		SocketDescription debug description of socket
-	 * @return		The new (heap-allocated) socket, or NULL if unsuccessful.
+	 * @param SocketDescription Debug description of socket,
+	 * @return The new (heap-allocated) socket, or nullptr if unsuccessful.
 	 */
 	virtual class FSocket* Accept(const FString& InSocketDescription) = 0;
 
 	/**
-	 * Accepts a connection that is pending
+	 * Accepts a connection that is pending.
 	 *
-	 * @param OutAddr the address of the connection
-	 * @param		SocketDescription debug description of socket
-	 *
-	 * @return		The new (heap-allocated) socket, or NULL if unsuccessful.
+	 * @param OutAddr The address of the connection.
+	 * @param SocketDescription Debug description of socket.
+	 * @return The new (heap-allocated) socket, or nullptr if unsuccessful.
 	 */
 	virtual class FSocket* Accept(FInternetAddr& OutAddr, const FString& InSocketDescription) = 0;
 
 	/**
-	 * Sends a buffer to a network byte ordered address
+	 * Sends a buffer to a network byte ordered address.
 	 *
-	 * @param Data the buffer to send
-	 * @param Count the size of the data to send
-	 * @param BytesSent out param indicating how much was sent
-	 * @param Destination the network byte ordered address to send to
+	 * @param Data The buffer to send.
+	 * @param Count The size of the data to send.
+	 * @param BytesSent Will indicate how much was sent.
+	 * @param Destination The network byte ordered address to send to.
 	 */
 	virtual bool SendTo(const uint8* Data, int32 Count, int32& BytesSent, const FInternetAddr& Destination);
 
 	/**
-	 * Sends a buffer on a connected socket
+	 * Sends a buffer on a connected socket.
 	 *
-	 * @param Data the buffer to send
-	 * @param Count the size of the data to send
-	 * @param BytesSent out param indicating how much was sent
+	 * @param Data The buffer to send.
+	 * @param Count The size of the data to send.
+	 * @param BytesSent Will indicate how much was sent.
 	 */
 	virtual bool Send(const uint8* Data, int32 Count, int32& BytesSent);
 
 	/**
-	 * Reads a chunk of data from the socket. Gathers the source address too
+	 * Reads a chunk of data from the socket and gathers the source address.
 	 *
-	 * @param Data the buffer to read into
-	 * @param BufferSize the max size of the buffer
-	 * @param BytesRead out param indicating how many bytes were read from the socket
-	 * @param Source out param receiving the address of the sender of the data
-	 * @param Flags the receive flags
+	 * @param Data The buffer to read into.
+	 * @param BufferSize The max size of the buffer.
+	 * @param BytesRead Will indicate how many bytes were read from the socket.
+	 * @param Source Will contain the receiving the address of the sender of the data.
+	 * @param Flags The receive flags.
 	 */
 	virtual bool RecvFrom(uint8* Data, int32 BufferSize, int32& BytesRead, FInternetAddr& Source, ESocketReceiveFlags::Type Flags = ESocketReceiveFlags::None);
 
 	/**
 	 * Reads a chunk of data from a connected socket
 	 *
-	 * @param Data the buffer to read into
-	 * @param BufferSize the max size of the buffer
-	 * @param BytesRead out param indicating how many bytes were read from the socket
+	 * @param Data The buffer to read into
+	 * @param BufferSize The max size of the buffer
+	 * @param BytesRead Will indicate how many bytes were read from the socket
 	 * @param Flags the receive flags
 	 */
 	virtual bool Recv(uint8* Data, int32 BufferSize, int32& BytesRead, ESocketReceiveFlags::Type Flags = ESocketReceiveFlags::None);
@@ -158,48 +150,47 @@ public:
 	/**
 	 * Blocks until the specified condition is met.
 	 *
-	 * @param Condition - The condition to wait for.
-	 * @param WaitTime - The maximum time to wait.
-	 *
+	 * @param Condition The condition to wait for.
+	 * @param WaitTime The maximum time to wait.
 	 * @return true if the condition was met, false if the time limit expired or an error occurred.
 	 */
-	virtual bool Wait (ESocketWaitConditions::Type Condition, FTimespan WaitTime) = 0;
+	virtual bool Wait(ESocketWaitConditions::Type Condition, FTimespan WaitTime) = 0;
 
 	/**
-	 * Determines the connection state of the socket
+	 * Determines the connection state of the socket.
+	 *
+	 * @return Connection state.
 	 */
 	virtual ESocketConnectionState GetConnectionState() = 0;
 
 	/**
-	 * Reads the address the socket is bound to and returns it
+	 * Reads the address the socket is bound to and returns it.
 	 * 
-	 * @param OutAddr address the socket is bound to
+	 * @param OutAddr The address the socket is bound to.
 	 */
 	virtual void GetAddress(FInternetAddr& OutAddr) = 0;
 
 	/**
-	 * Reads the address of the peer the socket is connected to
+	 * Reads the address of the peer the socket is connected to.
 	 *
-	 * @param OutAddr address of the peer the socket is connected to
-	 * @return return true if the address was retrieved correctly 
+	 * @param OutAddr Address of the peer the socket is connected to.
+	 * @return true if the address was retrieved correctly, false otherwise.
 	 */
 	virtual bool GetPeerAddress(FInternetAddr& OutAddr) = 0;
 
 	/**
-	 * Sets this socket into non-blocking mode
-	 *
-	 * @param bIsNonBlocking whether to enable blocking or not
+	 * Sets this socket into non-blocking mode.
 	 *
-	 * @return true if successful, false otherwise
+	 * @param bIsNonBlocking Whether to enable blocking or not.
+	 * @return true if successful, false otherwise.
 	 */
 	virtual bool SetNonBlocking(bool bIsNonBlocking = true) = 0;
 
 	/**
-	 * Sets a socket into broadcast mode (UDP only)
+	 * Sets a socket into broadcast mode (UDP only).
 	 *
-	 * @param bAllowBroadcast whether to enable broadcast or not
-	 *
-	 * @return true if successful, false otherwise
+	 * @param bAllowBroadcast Whether to enable broadcast or not.
+	 * @return true if successful, false otherwise.
 	 */
 	virtual bool SetBroadcast(bool bAllowBroadcast = true) = 0;
 
@@ -208,28 +199,20 @@ public:
 	 *
 	 * The multicast group address must be in the range 224.0.0.0 to 239.255.255.255.
 	 *
-	 * @param GroupAddress - The IP address of the multicast group.
-	 *
+	 * @param GroupAddress The IP address of the multicast group.
 	 * @return true on success, false otherwise.
-	 *
-	 * @see LeaveMulticastGroup
-	 * @see SetMulticastLoopback
-	 * @see SetMulticastTtl
+	 * @see LeaveMulticastGroup, SetMulticastLoopback, SetMulticastTtl
 	 */
-	virtual bool JoinMulticastGroup (const FInternetAddr& GroupAddress) = 0;
+	virtual bool JoinMulticastGroup(const FInternetAddr& GroupAddress) = 0;
 
 	/**
 	 * Removes this UDP client from the specified multicast group.
 	 *
 	 * @param The multicast group address to leave.
-	 *
 	 * @return true on success, false otherwise.
-	 *
-	 * @see JoinMulticastGroup
-	 * @see SetMulticastLoopback
-	 * @see SetMulticastTtl
+	 * @see JoinMulticastGroup, SetMulticastLoopback, SetMulticastTtl
 	 */
-	virtual bool LeaveMulticastGroup (const FInternetAddr& GroupAddress) = 0;
+	virtual bool LeaveMulticastGroup(const FInternetAddr& GroupAddress) = 0;
 
 	/**
 	 * Enables or disables multicast loopback on the socket (UDP only).
@@ -239,13 +222,10 @@ public:
 	 * enabled. It must be enabled if more than one listener is present
 	 * on a host.
 	 *
-	 * @param bLoopback - Whether loopback should be enabled.
-	 *
-	 * @see LeaveMulticastGroup
-	 * @see JoinMulticastGroup
-	 * @see SetMulticastTtl
+	 * @param bLoopback Whether loopback should be enabled.
+	 * @see LeaveMulticastGroup, JoinMulticastGroup, SetMulticastTtl
 	 */
-	virtual bool SetMulticastLoopback (bool bLoopback) = 0;
+	virtual bool SetMulticastLoopback(bool bLoopback) = 0;
 
 	/**
 	 * Sets the time to live (TTL) for multicast datagrams.
@@ -255,69 +235,67 @@ public:
 	 * allow multicast datagrams to be sent into neighboring subnets, if
 	 * multicast capable routers are present.
 	 *
-	 * @param TimeToLive - Number of hops the datagram can make.
-	 *
-	 * @see LeaveMulticastGroup
-	 * @see JoinMulticastGroup
-	 * @see SetMulticastLoopback
+	 * @param TimeToLive Number of hops the datagram can make.
+	 * @see LeaveMulticastGroup, JoinMulticastGroup, SetMulticastLoopback
 	 */
-	virtual bool SetMulticastTtl (uint8 TimeToLive) = 0;
+	virtual bool SetMulticastTtl(uint8 TimeToLive) = 0;
 
 	/**
-	 * Sets whether a socket can be bound to an address in use
-	 *
-	 * @param bAllowReuse whether to allow reuse or not
+	 * Sets whether a socket can be bound to an address in use.
 	 *
-	 * @return true if the call succeeded, false otherwise
+	 * @param bAllowReuse Whether to allow reuse or not.
+	 * @return true if the call succeeded, false otherwise.
 	 */
 	virtual bool SetReuseAddr(bool bAllowReuse = true) = 0;
 
 	/**
-	 * Sets whether and how long a socket will linger after closing
-	 *
-	 * @param bShouldLinger whether to have the socket remain open for a time period after closing or not
-	 * @param Timeout the amount of time to linger before closing
+	 * Sets whether and how long a socket will linger after closing.
 	 *
-	 * @return true if the call succeeded, false otherwise
+	 * @param bShouldLinger Whether to have the socket remain open for a time period after closing or not.
+	 * @param Timeout The amount of time to linger before closing.
+	 * @return true if the call succeeded, false otherwise.
 	 */
 	virtual bool SetLinger(bool bShouldLinger = true, int32 Timeout = 0) = 0;
 
 	/**
-	 * Enables error queue support for the socket
+	 * Enables error queue support for the socket.
 	 *
-	 * @param bUseErrorQueue whether to enable error queuing or not
-	 *
-	 * @return true if the call succeeded, false otherwise
+	 * @param bUseErrorQueue Whether to enable error queuing or not.
+	 * @return true if the call succeeded, false otherwise.
 	 */
 	virtual bool SetRecvErr(bool bUseErrorQueue = true) = 0;
 
 	/**
-	 * Sets the size of the send buffer to use
-	 *
-	 * @param Size the size to change it to
-	 * @param NewSize the out value returning the size that was set (in case OS can't set that)
+	 * Sets the size of the send buffer to use.
 	 *
-	 * @return true if the call succeeded, false otherwise
+	 * @param Size The size to change it to.
+	 * @param NewSize Will contain the size that was set (in case OS can't set that).
+	 * @return true if the call succeeded, false otherwise.
 	 */
 	virtual bool SetSendBufferSize(int32 Size, int32& NewSize) = 0;
 
 	/**
-	 * Sets the size of the receive buffer to use
+	 * Sets the size of the receive buffer to use.
 	 *
-	 * @param Size the size to change it to
-	 * @param NewSize the out value returning the size that was set (in case OS can't set that)
-	 *
-	 * @return true if the call succeeded, false otherwise
+	 * @param Size The size to change it to.
+	 * @param NewSize Will contain the size that was set (in case OS can't set that).
+	 * @return true if the call succeeded, false otherwise.
 	 */
 	virtual bool SetReceiveBufferSize(int32 Size, int32& NewSize) = 0;
 
 	/**
 	 * Reads the port this socket is bound to.
+	 *
+	 * @return Port number.
+	 * @see GetSocketType, GetDescription
 	 */ 
 	virtual int32 GetPortNo() = 0;
 
 	/**
-	 * @return The type of protocol the socket is bound to
+	 * Get the type of protocol the socket is bound to
+	 *
+	 * @return Socket type.
+	 * @see GetDescription, GetPortNo
 	 */
 	FORCEINLINE ESocketType GetSocketType() const
 	{
@@ -325,10 +303,13 @@ public:
 	}
 
 	/**
-	 * @return The debug description of the socket
+	 * Get the debug description of the socket.
+	 *
+	 * @return Debug string.
+	 * @see GetPortNo, GetSocketType
 	 */
 	FORCEINLINE FString GetDescription() const
 	{
 		return SocketDescription;
 	}
-};
\ No newline at end of file
+};
diff --git a/Engine/Source/Runtime/UMG/Private/Animation/MovieScene2DTransformSection.cpp b/Engine/Source/Runtime/UMG/Private/Animation/MovieScene2DTransformSection.cpp
--- a/Engine/Source/Runtime/UMG/Private/Animation/MovieScene2DTransformSection.cpp
+++ b/Engine/Source/Runtime/UMG/Private/Animation/MovieScene2DTransformSection.cpp
@@ -133,7 +133,8 @@ void UMovieScene2DTransformSection::SetKeyTime( FKeyHandle KeyHandle, float Time
 /**
  * Chooses an appropriate curve from an axis and a set of curves
  */
-static FRichCurve* ChooseCurve( EAxis::Type Axis, FRichCurve* Curves )
+template<typename T>
+static T* ChooseCurve( EAxis::Type Axis, T* Curves )
 {
 	switch(Axis)
 	{
@@ -156,19 +157,43 @@ FRichCurve& UMovieScene2DTransformSection::GetTranslationCurve(EAxis::Type Axis)
 }
 
 
+const FRichCurve& UMovieScene2DTransformSection::GetTranslationCurve(EAxis::Type Axis) const
+{
+	return *ChooseCurve( Axis, Translation );
+}
+
+
 FRichCurve& UMovieScene2DTransformSection::GetRotationCurve()
 {
 	return Rotation;
 }
 
 
+const FRichCurve& UMovieScene2DTransformSection::GetRotationCurve() const
+{
+	return Rotation;
+}
+
+
 FRichCurve& UMovieScene2DTransformSection::GetScaleCurve(EAxis::Type Axis)
 {
 	return *ChooseCurve(Axis, Scale);
 }
 
 
-FRichCurve& UMovieScene2DTransformSection::GetSheerCurve(EAxis::Type Axis)
+const FRichCurve& UMovieScene2DTransformSection::GetScaleCurve(EAxis::Type Axis) const
+{
+	return *ChooseCurve(Axis, Scale);
+}
+
+
+FRichCurve& UMovieScene2DTransformSection::GetShearCurve(EAxis::Type Axis)
+{
+	return *ChooseCurve(Axis, Shear);
+}
+
+
+const FRichCurve& UMovieScene2DTransformSection::GetShearCurve(EAxis::Type Axis) const
 {
 	return *ChooseCurve(Axis, Shear);
 }
diff --git a/Engine/Source/Runtime/UMG/Private/Animation/MovieScene2DTransformTemplate.cpp b/Engine/Source/Runtime/UMG/Private/Animation/MovieScene2DTransformTemplate.cpp
new file mode 100644
--- /dev/null
+++ b/Engine/Source/Runtime/UMG/Private/Animation/MovieScene2DTransformTemplate.cpp
@@ -0,0 +1,48 @@
+// Copyright 1998-2016 Epic Games, Inc. All Rights Reserved.
+
+#include "UMGPrivatePCH.h"
+#include "MovieScene2DTransformTemplate.h"
+
+#include "MovieScene2DTransformSection.h"
+#include "MovieScene2DTransformTrack.h"
+
+FMovieScene2DTransformSectionTemplate::FMovieScene2DTransformSectionTemplate(const UMovieScene2DTransformSection& Section, const UMovieScenePropertyTrack& Track)
+	: PropertyData(Track.GetPropertyName(), Track.GetPropertyPath())
+{
+	Translation[0] = Section.GetTranslationCurve(EAxis::X);
+	Translation[1] = Section.GetTranslationCurve(EAxis::Y);
+
+	Rotation = Section.GetRotationCurve();
+
+	Scale[0] = Section.GetScaleCurve(EAxis::X);
+	Scale[1] = Section.GetScaleCurve(EAxis::Y);
+
+	Shear[0] = Section.GetShearCurve(EAxis::X);
+	Shear[1] = Section.GetShearCurve(EAxis::Y);
+}
+
+void FMovieScene2DTransformSectionTemplate::Evaluate(const FMovieSceneEvaluationOperand& Operand, const FMovieSceneContext& Context, const FPersistentEvaluationData& PersistentData, FMovieSceneExecutionTokens& ExecutionTokens) const
+{
+	using namespace PropertyTemplate;
+
+	TCachedSectionData<FWidgetTransform>& SectionData = PersistentData.GetSectionData<TCachedSectionData<FWidgetTransform>>();
+
+	float Time = Context.GetTime();
+	for (TCachedValue<FWidgetTransform>& ObjectAndValue : SectionData.ObjectsAndValues)
+	{
+		FWidgetTransform& Value = ObjectAndValue.Value;
+
+		Value.Translation.X 	= Translation[0].Eval(Time, Value.Translation.X);
+		Value.Translation.Y 	= Translation[1].Eval(Time, Value.Translation.Y);
+
+		Value.Scale.X 			= Scale[0].Eval(Time, Value.Scale.X);
+		Value.Scale.Y 			= Scale[1].Eval(Time, Value.Scale.Y);
+
+		Value.Shear.X 			= Shear[0].Eval(Time, Value.Shear.X);
+		Value.Shear.Y 			= Shear[1].Eval(Time, Value.Shear.Y);
+
+		Value.Angle 			= Rotation.Eval(Time, Value.Angle);
+	}
+
+	ExecutionTokens.Add(TCachedPropertyTrackExecutionToken<FWidgetTransform>());
+}
\ No newline at end of file
diff --git a/Engine/Source/Runtime/UMG/Private/Animation/MovieScene2DTransformTemplate.h b/Engine/Source/Runtime/UMG/Private/Animation/MovieScene2DTransformTemplate.h
new file mode 100644
--- /dev/null
+++ b/Engine/Source/Runtime/UMG/Private/Animation/MovieScene2DTransformTemplate.h
@@ -0,0 +1,58 @@
+// Copyright 1998-2016 Epic Games, Inc. All Rights Reserved.
+
+#pragma once
+
+#include "MovieScenePropertyTemplate.h"
+#include "MovieScene2DTransformTemplate.generated.h"
+
+class UMovieScene3DTransformSection;
+class UMovieScenePropertyTrack;
+
+USTRUCT()
+struct FMovieScene2DTransformSectionTemplate : public FMovieSceneEvalTemplate
+{
+	GENERATED_BODY()
+	
+	FMovieScene2DTransformSectionTemplate(){}
+	FMovieScene2DTransformSectionTemplate(const UMovieScene2DTransformSection& Section, const UMovieScenePropertyTrack& Track);
+
+private:
+
+	virtual UScriptStruct& GetScriptStructImpl() const override
+	{
+		return *StaticStruct();
+	}
+	virtual void SetupOverrides() override
+	{
+		EnableOverrides(RequiresSetupFlag | RequiresInitializeFlag);
+	}
+	virtual void Setup(FPersistentEvaluationData& PersistentData, IMovieScenePlayer& Player) const override
+	{
+		PropertyData.SetupCachedTrack<FWidgetTransform>(PersistentData);
+	}
+	virtual void Initialize(const FMovieSceneEvaluationOperand& Operand, const FMovieSceneContext& Context, FPersistentEvaluationData& PersistentData, IMovieScenePlayer& Player) const override
+	{
+		PropertyData.SetupCachedFrame<FWidgetTransform>(Operand, PersistentData, Player);
+	}
+
+	virtual void Evaluate(const FMovieSceneEvaluationOperand& Operand, const FMovieSceneContext& Context, const FPersistentEvaluationData& PersistentData, FMovieSceneExecutionTokens& ExecutionTokens) const override;
+
+	UPROPERTY()
+	FMovieScenePropertySectionData PropertyData;
+
+	/** Translation curves */
+	UPROPERTY()
+	FRichCurve Translation[2];
+	
+	/** Rotation curve */
+	UPROPERTY()
+	FRichCurve Rotation;
+
+	/** Scale curves */
+	UPROPERTY()
+	FRichCurve Scale[2];
+
+	/** Shear curve */
+	UPROPERTY()
+	FRichCurve Shear[2];
+};
\ No newline at end of file
diff --git a/Engine/Source/Runtime/UMG/Private/Animation/MovieScene2DTransformTrack.cpp b/Engine/Source/Runtime/UMG/Private/Animation/MovieScene2DTransformTrack.cpp
--- a/Engine/Source/Runtime/UMG/Private/Animation/MovieScene2DTransformTrack.cpp
+++ b/Engine/Source/Runtime/UMG/Private/Animation/MovieScene2DTransformTrack.cpp
@@ -4,7 +4,7 @@
 #include "MovieScene2DTransformSection.h"
 #include "MovieScene2DTransformTrack.h"
 #include "IMovieScenePlayer.h"
-#include "MovieScene2DTransformTrackInstance.h"
+#include "MovieScene2DTransformTemplate.h"
 #include "MovieSceneCommonHelpers.h"
 
 UMovieScene2DTransformTrack::UMovieScene2DTransformTrack(const FObjectInitializer& ObjectInitializer)
@@ -13,6 +13,8 @@ UMovieScene2DTransformTrack::UMovieScene2DTransformTrack(const FObjectInitialize
 #if WITH_EDITORONLY_DATA
 	TrackTint = FColor(48, 227, 255, 65);
 #endif
+
+	EvalOptions.bEvaluateNearestSection = EvalOptions.bCanEvaluateNearestSection = true;
 }
 
 UMovieSceneSection* UMovieScene2DTransformTrack::CreateNewSection()
@@ -21,12 +23,12 @@ UMovieSceneSection* UMovieScene2DTransformTrack::CreateNewSection()
 }
 
 
-TSharedPtr<IMovieSceneTrackInstance> UMovieScene2DTransformTrack::CreateInstance()
+FMovieSceneEvalTemplatePtr UMovieScene2DTransformTrack::CreateTemplateForSection(const UMovieSceneSection& InSection) const
 {
-	return MakeShareable( new FMovieScene2DTransformTrackInstance( *this ) );
+	return FMovieScene2DTransformSectionTemplate(*CastChecked<const UMovieScene2DTransformSection>(&InSection), *this);
 }
 
-
+PRAGMA_DISABLE_DEPRECATION_WARNINGS
 bool UMovieScene2DTransformTrack::Eval(float Position, float LastPosition, FWidgetTransform& InOutTransform) const
 {
 	const UMovieSceneSection* Section = MovieSceneHelpers::FindNearestSectionAtTime(Sections, Position);
@@ -45,3 +47,4 @@ bool UMovieScene2DTransformTrack::Eval(float Position, float LastPosition, FWidg
 
 	return (Section != nullptr);
 }
+PRAGMA_ENABLE_DEPRECATION_WARNINGS
\ No newline at end of file
diff --git a/Engine/Source/Runtime/UMG/Private/Animation/MovieSceneMarginSection.cpp b/Engine/Source/Runtime/UMG/Private/Animation/MovieSceneMarginSection.cpp
--- a/Engine/Source/Runtime/UMG/Private/Animation/MovieSceneMarginSection.cpp
+++ b/Engine/Source/Runtime/UMG/Private/Animation/MovieSceneMarginSection.cpp
@@ -115,15 +115,6 @@ void UMovieSceneMarginSection::SetKeyTime( FKeyHandle KeyHandle, float Time )
 }
 
 
-FMargin UMovieSceneMarginSection::Eval( float Position, const FMargin& DefaultValue ) const
-{
-	return FMargin(	LeftCurve.Eval(Position, DefaultValue.Left),
-					TopCurve.Eval(Position, DefaultValue.Top),
-					RightCurve.Eval(Position, DefaultValue.Right),
-					BottomCurve.Eval(Position, DefaultValue.Bottom));
-}
-
-
 template<typename CurveType>
 CurveType* GetCurveForChannel( EKeyMarginChannel Channel, CurveType* LeftCurve, CurveType* TopCurve, CurveType* RightCurve, CurveType* BottomCurve )
 {
diff --git a/Engine/Source/Runtime/UMG/Private/Animation/MovieSceneMarginTemplate.cpp b/Engine/Source/Runtime/UMG/Private/Animation/MovieSceneMarginTemplate.cpp
new file mode 100644
--- /dev/null
+++ b/Engine/Source/Runtime/UMG/Private/Animation/MovieSceneMarginTemplate.cpp
@@ -0,0 +1,36 @@
+// Copyright 1998-2016 Epic Games, Inc. All Rights Reserved.
+
+#include "UMGPrivatePCH.h"
+#include "MovieSceneMarginTemplate.h"
+
+#include "MovieSceneMarginSection.h"
+#include "MovieScenePropertyTrack.h"
+
+FMovieSceneMarginSectionTemplate::FMovieSceneMarginSectionTemplate(const UMovieSceneMarginSection& Section, const UMovieScenePropertyTrack& Track)
+	: PropertyData(Track.GetPropertyName(), Track.GetPropertyPath())
+	, TopCurve(Section.GetTopCurve())
+	, LeftCurve(Section.GetLeftCurve())
+	, RightCurve(Section.GetRightCurve())
+	, BottomCurve(Section.GetBottomCurve())
+{
+}
+
+void FMovieSceneMarginSectionTemplate::Evaluate(const FMovieSceneEvaluationOperand& Operand, const FMovieSceneContext& Context, const FPersistentEvaluationData& PersistentData, FMovieSceneExecutionTokens& ExecutionTokens) const
+{
+	using namespace PropertyTemplate;
+
+	TCachedSectionData<FMargin>& SectionData = PersistentData.GetSectionData<TCachedSectionData<FMargin>>();
+
+	float Time = Context.GetTime();
+	for (TCachedValue<FMargin>& ObjectAndValue : SectionData.ObjectsAndValues)
+	{
+		FMargin& Margin = ObjectAndValue.Value;
+
+		Margin.Top = TopCurve.Eval(Time, Margin.Top);
+		Margin.Left = LeftCurve.Eval(Time, Margin.Left);
+		Margin.Right = RightCurve.Eval(Time, Margin.Right);
+		Margin.Bottom = BottomCurve.Eval(Time, Margin.Bottom);
+	}
+
+	ExecutionTokens.Add(TCachedPropertyTrackExecutionToken<FMargin>());
+}
\ No newline at end of file
diff --git a/Engine/Source/Runtime/UMG/Private/Animation/MovieSceneMarginTemplate.h b/Engine/Source/Runtime/UMG/Private/Animation/MovieSceneMarginTemplate.h
new file mode 100644
--- /dev/null
+++ b/Engine/Source/Runtime/UMG/Private/Animation/MovieSceneMarginTemplate.h
@@ -0,0 +1,54 @@
+// Copyright 1998-2016 Epic Games, Inc. All Rights Reserved.
+
+#pragma once
+
+#include "MovieScenePropertyTemplate.h"
+#include "MovieSceneMarginTemplate.generated.h"
+
+class UMovieScene3DTransformSection;
+class UMovieScenePropertyTrack;
+
+USTRUCT()
+struct FMovieSceneMarginSectionTemplate : public FMovieSceneEvalTemplate
+{
+	GENERATED_BODY()
+	
+	FMovieSceneMarginSectionTemplate(){}
+	FMovieSceneMarginSectionTemplate(const UMovieSceneMarginSection& Section, const UMovieScenePropertyTrack& Track);
+
+private:
+
+	virtual UScriptStruct& GetScriptStructImpl() const override
+	{
+		return *StaticStruct();
+	}
+	virtual void SetupOverrides() override
+	{
+		EnableOverrides(RequiresSetupFlag | RequiresInitializeFlag);
+	}
+	virtual void Setup(FPersistentEvaluationData& PersistentData, IMovieScenePlayer& Player) const override
+	{
+		PropertyData.SetupCachedTrack<FMargin>(PersistentData);
+	}
+	virtual void Initialize(const FMovieSceneEvaluationOperand& Operand, const FMovieSceneContext& Context, FPersistentEvaluationData& PersistentData, IMovieScenePlayer& Player) const override
+	{
+		PropertyData.SetupCachedFrame<FMargin>(Operand, PersistentData, Player);
+	}
+
+	virtual void Evaluate(const FMovieSceneEvaluationOperand& Operand, const FMovieSceneContext& Context, const FPersistentEvaluationData& PersistentData, FMovieSceneExecutionTokens& ExecutionTokens) const override;
+
+	UPROPERTY()
+	FMovieScenePropertySectionData PropertyData;
+
+	UPROPERTY()
+	FRichCurve TopCurve;
+
+	UPROPERTY()
+	FRichCurve LeftCurve;
+
+	UPROPERTY()
+	FRichCurve RightCurve;
+
+	UPROPERTY()
+	FRichCurve BottomCurve;
+};
\ No newline at end of file
diff --git a/Engine/Source/Runtime/UMG/Private/Animation/MovieSceneMarginTrack.cpp b/Engine/Source/Runtime/UMG/Private/Animation/MovieSceneMarginTrack.cpp
--- a/Engine/Source/Runtime/UMG/Private/Animation/MovieSceneMarginTrack.cpp
+++ b/Engine/Source/Runtime/UMG/Private/Animation/MovieSceneMarginTrack.cpp
@@ -4,37 +4,23 @@
 #include "MovieSceneMarginSection.h"
 #include "MovieSceneMarginTrack.h"
 #include "IMovieScenePlayer.h"
-#include "MovieSceneMarginTrackInstance.h"
+#include "MovieSceneMarginTemplate.h"
 #include "MovieSceneCommonHelpers.h"
 
 
-UMovieSceneSection* UMovieSceneMarginTrack::CreateNewSection()
+UMovieSceneMarginTrack::UMovieSceneMarginTrack(const FObjectInitializer& Init)
+	: Super(Init)
 {
-	return NewObject<UMovieSceneSection>(this, UMovieSceneMarginSection::StaticClass(), NAME_None, RF_Transactional);
+	EvalOptions.bEvaluateNearestSection = EvalOptions.bCanEvaluateNearestSection = true;
 }
 
-
-TSharedPtr<IMovieSceneTrackInstance> UMovieSceneMarginTrack::CreateInstance()
+UMovieSceneSection* UMovieSceneMarginTrack::CreateNewSection()
 {
-	return MakeShareable( new FMovieSceneMarginTrackInstance( *this ) ); 
+	return NewObject<UMovieSceneSection>(this, UMovieSceneMarginSection::StaticClass(), NAME_None, RF_Transactional);
 }
 
 
-bool UMovieSceneMarginTrack::Eval( float Position, float LastPosition, FMargin& InOutMargin ) const
+FMovieSceneEvalTemplatePtr UMovieSceneMarginTrack::CreateTemplateForSection(const UMovieSceneSection& InSection) const
 {
-	const UMovieSceneSection* Section = MovieSceneHelpers::FindNearestSectionAtTime( Sections, Position );
-
-	if( Section )
-	{
-		const UMovieSceneMarginSection* MarginSection = CastChecked<UMovieSceneMarginSection>( Section );
-
-		if (!Section->IsInfinite())
-		{
-			Position = FMath::Clamp(Position, Section->GetStartTime(), Section->GetEndTime());
-		}
-
-		InOutMargin = MarginSection->Eval( Position, InOutMargin );
-	}
-
-	return (Section != nullptr);
+	return FMovieSceneMarginSectionTemplate(*CastChecked<UMovieSceneMarginSection>(&InSection), *this);
 }
diff --git a/Engine/Source/Runtime/UMG/Private/Animation/MovieSceneWidgetMaterialTemplate.cpp b/Engine/Source/Runtime/UMG/Private/Animation/MovieSceneWidgetMaterialTemplate.cpp
new file mode 100644
--- /dev/null
+++ b/Engine/Source/Runtime/UMG/Private/Animation/MovieSceneWidgetMaterialTemplate.cpp
@@ -0,0 +1,82 @@
+// Copyright 1998-2016 Epic Games, Inc. All Rights Reserved.
+
+#include "UMGPrivatePCH.h"
+#include "MovieSceneWidgetMaterialTemplate.h"
+#include "WidgetMaterialTrackUtilities.h"
+#include "MovieSceneWidgetMaterialTrack.h"
+
+// Container to ensure unique IDs per property path
+TMovieSceneAnimTypeIDContainer<TArray<FName>> BrushPropertyIDs;
+
+struct FWidgetMaterialAccessor : FDefaultMaterialAccessor
+{
+	FWidgetMaterialAccessor(const TArray<FName>& InBrushPropertyNamePath)
+		: AnimTypeID(BrushPropertyIDs.GetAnimTypeID(InBrushPropertyNamePath))
+		, BrushPropertyNamePath(InBrushPropertyNamePath)
+	{}
+
+	FWidgetMaterialAccessor(const FWidgetMaterialAccessor&) = default;
+	FWidgetMaterialAccessor& operator=(const FWidgetMaterialAccessor&) = default;
+
+#if PLATFORM_COMPILER_HAS_DEFAULTED_FUNCTIONS
+	FWidgetMaterialAccessor(FWidgetMaterialAccessor&&) = default;
+	FWidgetMaterialAccessor& operator=(FWidgetMaterialAccessor&&) = default;
+#else
+	FWidgetMaterialAccessor(FWidgetMaterialAccessor&& RHS)
+		: AnimTypeID(MoveTemp(RHS.AnimTypeID))
+		, BrushPropertyNamePath(MoveTemp(RHS.BrushPropertyNamePath))
+	{
+	}
+	FWidgetMaterialAccessor& operator=(FWidgetMaterialAccessor&& RHS)
+	{
+		AnimTypeID = MoveTemp(RHS.AnimTypeID);
+		BrushPropertyNamePath = MoveTemp(RHS.BrushPropertyNamePath);
+		return *this;
+	}
+#endif
+
+	FMovieSceneAnimTypeID GetAnimTypeID() const
+	{
+		return AnimTypeID;
+	}
+
+	UMaterialInterface* GetMaterialForObject(UObject& Object)
+	{
+		if (UWidget* Widget = Cast<UWidget>(&Object))
+		{
+			FWidgetMaterialHandle Handle = WidgetMaterialTrackUtilities::GetMaterialHandle( Widget, BrushPropertyNamePath );
+			if (Handle.IsValid())
+			{
+				return Handle.GetMaterial();
+			}
+		}
+		return nullptr;
+	}
+
+	void SetMaterialForObject(UObject& Object, UMaterialInterface& Material)
+	{
+		if (UWidget* Widget = Cast<UWidget>(&Object))
+		{
+			FWidgetMaterialHandle Handle = WidgetMaterialTrackUtilities::GetMaterialHandle(Widget, BrushPropertyNamePath);
+			Handle.SetMaterial(&Material);
+		}
+	}
+
+	FMovieSceneAnimTypeID AnimTypeID;
+	TArray<FName> BrushPropertyNamePath;
+};
+
+FMovieSceneWidgetMaterialSectionTemplate::FMovieSceneWidgetMaterialSectionTemplate(const UMovieSceneParameterSection& Section, const UMovieSceneWidgetMaterialTrack& Track)
+	: FMovieSceneParameterSectionTemplate(Section)
+	, BrushPropertyNamePath(Track.GetBrushPropertyNamePath())
+{
+}
+
+void FMovieSceneWidgetMaterialSectionTemplate::Evaluate(const FMovieSceneEvaluationOperand& Operand, const FMovieSceneContext& Context, const FPersistentEvaluationData& PersistentData, FMovieSceneExecutionTokens& ExecutionTokens) const
+{
+	TMaterialTrackExecutionToken<FWidgetMaterialAccessor> ExecutionToken(BrushPropertyNamePath);
+
+	EvaluateCurves(Context, ExecutionToken.Values);
+
+	ExecutionTokens.Add(MoveTemp(ExecutionToken));
+}
\ No newline at end of file
diff --git a/Engine/Source/Runtime/UMG/Private/Animation/MovieSceneWidgetMaterialTemplate.h b/Engine/Source/Runtime/UMG/Private/Animation/MovieSceneWidgetMaterialTemplate.h
new file mode 100644
--- /dev/null
+++ b/Engine/Source/Runtime/UMG/Private/Animation/MovieSceneWidgetMaterialTemplate.h
@@ -0,0 +1,25 @@
+// Copyright 1998-2016 Epic Games, Inc. All Rights Reserved.
+
+#pragma once
+
+#include "Evaluation/MovieSceneParameterTemplate.h"
+#include "MovieSceneWidgetMaterialTemplate.generated.h"
+
+class UMovieSceneWidgetMaterialTrack;
+
+USTRUCT()
+struct FMovieSceneWidgetMaterialSectionTemplate : public FMovieSceneParameterSectionTemplate
+{
+	GENERATED_BODY()
+
+	FMovieSceneWidgetMaterialSectionTemplate() {}
+	FMovieSceneWidgetMaterialSectionTemplate(const UMovieSceneParameterSection& Section, const UMovieSceneWidgetMaterialTrack& Track);
+
+private:
+
+	virtual UScriptStruct& GetScriptStructImpl() const override { return *StaticStruct(); }
+	virtual void Evaluate(const FMovieSceneEvaluationOperand& Operand, const FMovieSceneContext& Context, const FPersistentEvaluationData& PersistentData, FMovieSceneExecutionTokens& ExecutionTokens) const override;
+
+	UPROPERTY()
+	TArray<FName> BrushPropertyNamePath;
+};
\ No newline at end of file
diff --git a/Engine/Source/Runtime/UMG/Private/Animation/MovieSceneWidgetMaterialTrack.cpp b/Engine/Source/Runtime/UMG/Private/Animation/MovieSceneWidgetMaterialTrack.cpp
--- a/Engine/Source/Runtime/UMG/Private/Animation/MovieSceneWidgetMaterialTrack.cpp
+++ b/Engine/Source/Runtime/UMG/Private/Animation/MovieSceneWidgetMaterialTrack.cpp
@@ -2,8 +2,8 @@
 
 #include "UMGPrivatePCH.h"
 #include "MovieSceneWidgetMaterialTrack.h"
-#include "MovieSceneWidgetMaterialTrackInstance.h"
 #include "WidgetMaterialTrackUtilities.h"
+#include "MovieSceneWidgetMaterialTemplate.h"
 
 
 UMovieSceneWidgetMaterialTrack::UMovieSceneWidgetMaterialTrack( const FObjectInitializer& ObjectInitializer )
@@ -12,9 +12,9 @@ UMovieSceneWidgetMaterialTrack::UMovieSceneWidgetMaterialTrack( const FObjectIni
 }
 
 
-TSharedPtr<IMovieSceneTrackInstance> UMovieSceneWidgetMaterialTrack::CreateInstance()
+FMovieSceneEvalTemplatePtr UMovieSceneWidgetMaterialTrack::CreateTemplateForSection(const UMovieSceneSection& InSection) const
 {
-	return MakeShareable(new FMovieSceneWidgetMaterialTrackInstance(*this));
+	return FMovieSceneWidgetMaterialSectionTemplate(*CastChecked<UMovieSceneParameterSection>(&InSection), *this);
 }
 
 
diff --git a/Engine/Source/Runtime/UMG/Private/Animation/UMGSequencePlayer.cpp b/Engine/Source/Runtime/UMG/Private/Animation/UMGSequencePlayer.cpp
--- a/Engine/Source/Runtime/UMG/Private/Animation/UMGSequencePlayer.cpp
+++ b/Engine/Source/Runtime/UMG/Private/Animation/UMGSequencePlayer.cpp
@@ -18,7 +18,7 @@ UUMGSequencePlayer::UUMGSequencePlayer(const FObjectInitializer& ObjectInitializ
 	Animation = nullptr;
 }
 
-void UUMGSequencePlayer::InitSequencePlayer( const UWidgetAnimation& InAnimation, UUserWidget& InUserWidget )
+void UUMGSequencePlayer::InitSequencePlayer( UWidgetAnimation& InAnimation, UUserWidget& InUserWidget )
 {
 	Animation = &InAnimation;
 
@@ -29,24 +29,7 @@ void UUMGSequencePlayer::InitSequencePlayer( const UWidgetAnimation& InAnimation
 	AnimationStartOffset = TimeRange.GetLowerBoundValue();
 
 	UserWidget = &InUserWidget;
-	UWidgetTree* WidgetTree = InUserWidget.WidgetTree;
-
-	// Bind to Runtime Objects
-	for (const FWidgetAnimationBinding& Binding : InAnimation.GetBindings())
-	{
-		UObject* FoundObject = Binding.FindRuntimeObject( *WidgetTree , InUserWidget);
-
-		if( FoundObject )
-		{
-			TArray<UObject*>& Objects = GuidToRuntimeObjectMap.FindOrAdd(Binding.AnimationGuid);
-			Objects.Add(FoundObject);
-		}
-		else
-		{
-			UE_LOG(LogUMG, Warning, TEXT("Failed to find runtime objects for %s animation, WidgetName: %s, SlotName: %s"), *InAnimation.GetPathName(), *Binding.WidgetName.ToString(), *Binding.SlotWidgetName.ToString() );
 		}
-	}
-}
 
 void UUMGSequencePlayer::Tick(float DeltaTime)
 {
@@ -127,18 +110,19 @@ void UUMGSequencePlayer::Tick(float DeltaTime)
 			Animation->OnAnimationFinished.Broadcast();
 		}
 
-		if (RootMovieSceneInstance.IsValid())
+		if (RootTemplateInstance.IsValid())
 		{			
-			EMovieSceneUpdateData UpdateData(TimeCursorPosition + AnimationStartOffset, LastTimePosition + AnimationStartOffset);
-			RootMovieSceneInstance->Update(UpdateData, *this);
+			const FMovieSceneContext Context(
+				FMovieSceneEvaluationRange(TimeCursorPosition + AnimationStartOffset, (TimeCursorPosition - LastTimePosition) + AnimationStartOffset),
+				PlayerStatus);
+			RootTemplateInstance.Evaluate(Context, *this);
 		}
 	}
 }
 
 void UUMGSequencePlayer::PlayInternal(double StartAtTime, double EndAtTime, double SubAnimStartTime, double SubAnimEndTime, int32 InNumLoopsToPlay, EUMGSequencePlayMode::Type InPlayMode, float InPlaybackSpeed)
 {
-	RootMovieSceneInstance = MakeShareable( new FMovieSceneSequenceInstance( *Animation ) );
-	RootMovieSceneInstance->RefreshInstance( *this );
+	RootTemplateInstance.Initialize(*Animation, *this);
 
 	PlaybackSpeed = FMath::Abs(InPlaybackSpeed);
 	PlayMode = InPlayMode;
@@ -176,10 +160,10 @@ void UUMGSequencePlayer::PlayInternal(double StartAtTime, double EndAtTime, doub
 	// Immediately evaluate the first frame of the animation so that if tick has already occurred, the widget is setup correctly and ready to be
 	// rendered using the first frames data, otherwise you may see a *pop* due to a widget being constructed with a default different than the
 	// first frame of the animation.
-	if ( RootMovieSceneInstance.IsValid() )
-	{
-		EMovieSceneUpdateData UpdateData(TimeCursorPosition, TimeCursorPosition);
-		RootMovieSceneInstance->Update(UpdateData, *this);
+	if (RootTemplateInstance.IsValid())
+	{			
+		const FMovieSceneContext Context(FMovieSceneEvaluationRange(TimeCursorPosition, TimeCursorPosition), PlayerStatus);
+		RootTemplateInstance.Evaluate(Context, *this);
 	}
 
 	PlayerStatus = EMovieScenePlayerStatus::Playing;
@@ -206,6 +190,8 @@ void UUMGSequencePlayer::Pause()
 {
 	// Purposely don't trigger any OnFinished events
 	PlayerStatus = EMovieScenePlayerStatus::Stopped;
+
+	RootTemplateInstance.Finish(*this);
 }
 
 void UUMGSequencePlayer::Reverse()
@@ -220,6 +206,8 @@ void UUMGSequencePlayer::Stop()
 {
 	PlayerStatus = EMovieScenePlayerStatus::Stopped;
 
+	RootTemplateInstance.Finish(*this);
+
 	OnSequenceFinishedPlayingEvent.Broadcast(*this);
 	Animation->OnAnimationFinished.Broadcast();
 
@@ -243,16 +231,6 @@ void UUMGSequencePlayer::SetPlaybackSpeed(float InPlaybackSpeed)
 	PlaybackSpeed = InPlaybackSpeed;
 }
 
-void UUMGSequencePlayer::GetRuntimeObjects(TSharedRef<FMovieSceneSequenceInstance> MovieSceneInstance, const FGuid& ObjectHandle, TArray<TWeakObjectPtr<UObject>>& OutObjects) const
-{
-	const TArray<UObject*>* FoundObjects = GuidToRuntimeObjectMap.Find( ObjectHandle );
-
-	if( FoundObjects )
-	{
-		OutObjects.Append(*FoundObjects);
-	}
-}
-
 EMovieScenePlayerStatus::Type UUMGSequencePlayer::GetPlaybackStatus() const
 {
 	return PlayerStatus;
@@ -260,14 +238,7 @@ EMovieScenePlayerStatus::Type UUMGSequencePlayer::GetPlaybackStatus() const
 
 UObject* UUMGSequencePlayer::GetPlaybackContext() const
 {
-	if (UserWidget.IsValid())
-	{
-		return UserWidget->GetWorld();
-	}
-	else
-	{
-		return nullptr;
-	}
+	return UserWidget.Get();
 }
 
 TArray<UObject*> UUMGSequencePlayer::GetEventContexts() const
diff --git a/Engine/Source/Runtime/UMG/Private/Animation/WidgetAnimation.cpp b/Engine/Source/Runtime/UMG/Private/Animation/WidgetAnimation.cpp
--- a/Engine/Source/Runtime/UMG/Private/Animation/WidgetAnimation.cpp
+++ b/Engine/Source/Runtime/UMG/Private/Animation/WidgetAnimation.cpp
@@ -15,7 +15,9 @@
 UWidgetAnimation::UWidgetAnimation(const FObjectInitializer& ObjectInitializer)
 	: Super(ObjectInitializer)
 	, MovieScene(nullptr)
-{ }
+{
+	bParentContextsAreSignificant = false;
+}
 
 
 /* UWidgetAnimation interface
@@ -53,61 +55,17 @@ float UWidgetAnimation::GetEndTime() const
 }
 
 
-void UWidgetAnimation::Initialize(UUserWidget* InPreviewWidget)
-{
-	PreviewWidget = InPreviewWidget;
-
-	// clear object maps
-	PreviewObjectToIds.Empty();
-	IdToPreviewObjects.Empty();
-	IdToSlotContentPreviewObjects.Empty();
-	SlotContentPreviewObjectToIds.Empty();
-
-	if (PreviewWidget == nullptr)
-	{
-		return;
-	}
-
-	// populate object maps
-	UWidgetTree* WidgetTree = PreviewWidget->WidgetTree;
-
-	for (const FWidgetAnimationBinding& Binding : AnimationBindings)
-	{
-		UObject* FoundObject = Binding.FindRuntimeObject(*WidgetTree, *InPreviewWidget);
-
-		if (FoundObject == nullptr)
-		{
-			continue;
-		}
-
-		UPanelSlot* FoundSlot = Cast<UPanelSlot>(FoundObject);
-
-		if (FoundSlot == nullptr)
-		{
-			IdToPreviewObjects.Add(Binding.AnimationGuid, FoundObject);
-			PreviewObjectToIds.Add(FoundObject, Binding.AnimationGuid);
-		}
-		else
-		{
-			IdToSlotContentPreviewObjects.Add(Binding.AnimationGuid, FoundSlot->Content);
-			SlotContentPreviewObjectToIds.Add(FoundSlot->Content, Binding.AnimationGuid);
-		}
-	}
-}
-
-
 /* UMovieSceneAnimation overrides
  *****************************************************************************/
 
 
 void UWidgetAnimation::BindPossessableObject(const FGuid& ObjectId, UObject& PossessedObject, UObject* Context)
 {
+	UUserWidget* PreviewWidget = CastChecked<UUserWidget>(Context);
+
 	// If it's the Root Widget
-	if (&PossessedObject == PreviewWidget.Get())
+	if (&PossessedObject == PreviewWidget)
 	{
-		PreviewObjectToIds.Add(&PossessedObject, ObjectId);
-		IdToPreviewObjects.Add(ObjectId, &PossessedObject);
-
 		FWidgetAnimationBinding NewBinding;
 		{
 			NewBinding.AnimationGuid = ObjectId;
@@ -123,9 +81,6 @@ void UWidgetAnimation::BindPossessableObject(const FGuid& ObjectId, UObject& Pos
 
 	if ((PossessedSlot != nullptr) && (PossessedSlot->Content != nullptr))
 	{
-		SlotContentPreviewObjectToIds.Add(PossessedSlot->Content, ObjectId);
-		IdToSlotContentPreviewObjects.Add(ObjectId, PossessedSlot->Content);
-
 		// Save the name of the widget containing the slots. This is the object
 		// to look up that contains the slot itself (the thing we are animating).
 		FWidgetAnimationBinding NewBinding;
@@ -140,9 +95,6 @@ void UWidgetAnimation::BindPossessableObject(const FGuid& ObjectId, UObject& Pos
 	}
 	else if (PossessedSlot == nullptr)
 	{
-		PreviewObjectToIds.Add(&PossessedObject, ObjectId);
-		IdToPreviewObjects.Add(ObjectId, &PossessedObject);
-
 		FWidgetAnimationBinding NewBinding;
 		{
 			NewBinding.AnimationGuid = ObjectId;
@@ -155,9 +107,11 @@ void UWidgetAnimation::BindPossessableObject(const FGuid& ObjectId, UObject& Pos
 }
 
 
-bool UWidgetAnimation::CanPossessObject(UObject& Object) const
+bool UWidgetAnimation::CanPossessObject(UObject& Object, UObject* InPlaybackContext) const
 {
-	if (&Object == PreviewWidget.Get())
+	UUserWidget* PreviewWidget = CastChecked<UUserWidget>(InPlaybackContext);
+
+	if (&Object == PreviewWidget)
 	{
 		return true;
 	}
@@ -170,45 +124,24 @@ bool UWidgetAnimation::CanPossessObject(UObject& Object) const
 		return false;
 	}
 
-	return (Object.IsA<UVisual>() && Object.IsIn(PreviewWidget.Get()));
+	return (Object.IsA<UVisual>() && Object.IsIn(PreviewWidget));
 }
 
-
-UObject* UWidgetAnimation::FindPossessableObject(const FGuid& ObjectId, UObject* Context) const
+void UWidgetAnimation::LocateBoundObjects(const FGuid& ObjectId, UObject* Context, TArray<UObject*, TInlineAllocator<1>>& OutObjects) const
 {
-	TWeakObjectPtr<UObject> PreviewObject = IdToPreviewObjects.FindRef(ObjectId);
-
-	if (PreviewObject.IsValid())
-	{
-		return PreviewObject.Get();
-	}
-
-	TWeakObjectPtr<UObject> SlotContentPreviewObject = IdToSlotContentPreviewObjects.FindRef(ObjectId);
-
-	if (SlotContentPreviewObject.IsValid())
-	{
-		UWidget* Widget = Cast<UWidget>(SlotContentPreviewObject.Get());
-		if ( Widget != nullptr )
-		{
-			return Widget->Slot;
+	const FWidgetAnimationBinding* Binding = AnimationBindings.FindByPredicate(
+		[ObjectId](const FWidgetAnimationBinding& In){
+			return In.AnimationGuid == ObjectId;
 		}
-	}
+	);
 
-	return nullptr;
-}
-
-
-FGuid UWidgetAnimation::FindPossessableObjectId(UObject& Object) const
-{
-	UPanelSlot* Slot = Cast<UPanelSlot>(&Object);
+	UUserWidget* PreviewWidget = CastChecked<UUserWidget>(Context);
+	UObject* FoundObject = Binding ? Binding->FindRuntimeObject(*PreviewWidget->WidgetTree, *PreviewWidget) : nullptr;
 
-	if (Slot != nullptr)
+	if (FoundObject)
 	{
-		// slot guids are tracked by their content.
-		return SlotContentPreviewObjectToIds.FindRef(Slot->Content);
+		OutObjects.Add(FoundObject);
 	}
-
-	return PreviewObjectToIds.FindRef(&Object);
 }
 
 
@@ -233,35 +166,8 @@ UObject* UWidgetAnimation::GetParentObject(UObject* Object) const
 	return nullptr;
 }
 
-
 void UWidgetAnimation::UnbindPossessableObjects(const FGuid& ObjectId)
 {
-	// unbind widgets
-	TArray<TWeakObjectPtr<UObject>> PreviewObjects;
-	{
-		IdToPreviewObjects.MultiFind(ObjectId, PreviewObjects);
-
-		for (TWeakObjectPtr<UObject>& PreviewObject : PreviewObjects)
-		{
-			PreviewObjectToIds.Remove(PreviewObject);
-		}
-
-		IdToPreviewObjects.Remove(ObjectId);
-	}
-
-	// unbind slots
-	TArray<TWeakObjectPtr<UObject>> SlotContentPreviewObjects;
-	{
-		IdToSlotContentPreviewObjects.MultiFind(ObjectId, SlotContentPreviewObjects);
-
-		for (TWeakObjectPtr<UObject>& SlotContentPreviewObject : SlotContentPreviewObjects)
-		{
-			SlotContentPreviewObjectToIds.Remove(SlotContentPreviewObject);
-		}
-
-		IdToSlotContentPreviewObjects.Remove(ObjectId);
-	}
-
 	// mark dirty
 	Modify();
 
@@ -271,38 +177,5 @@ void UWidgetAnimation::UnbindPossessableObjects(const FGuid& ObjectId)
 	});
 }
 
-void UWidgetAnimation::ReplacePossessableObject(const FGuid& OldId, const FGuid& NewId, UObject& OldObject, UObject& NewObject)
-{
-	TArray<TWeakObjectPtr<UObject>> PreviewObjects;
-	{
-		IdToPreviewObjects.MultiFind(OldId, PreviewObjects);
-
-		for (TWeakObjectPtr<UObject>& PreviewObject : PreviewObjects)
-		{
-			PreviewObjectToIds.Remove(PreviewObject);
-			PreviewObjectToIds.Add(&NewObject, NewId);
-		}
-
-		IdToPreviewObjects.Remove(OldId);
-		IdToPreviewObjects.Add(NewId, &NewObject);
-	}
-
-	// slots
-	FGuid SlotContentId = SlotContentPreviewObjectToIds.FindRef(&OldObject);
-	TArray<TWeakObjectPtr<UObject>> SlotContentPreviewObjects;
-	{
-		IdToSlotContentPreviewObjects.MultiFind(SlotContentId, SlotContentPreviewObjects);
-		
-		for (TWeakObjectPtr<UObject>& SlotContentPreviewObject : SlotContentPreviewObjects)
-		{
-			SlotContentPreviewObjectToIds.Remove(SlotContentPreviewObject);
-			SlotContentPreviewObjectToIds.Add(&NewObject, SlotContentId);
-		}
-
-		IdToSlotContentPreviewObjects.Remove(SlotContentId);
-		IdToSlotContentPreviewObjects.Add(SlotContentId, &NewObject);
-	}
-}
-
 
 #undef LOCTEXT_NAMESPACE
diff --git a/Engine/Source/Runtime/UMG/Private/UserWidget.cpp b/Engine/Source/Runtime/UMG/Private/UserWidget.cpp
--- a/Engine/Source/Runtime/UMG/Private/UserWidget.cpp
+++ b/Engine/Source/Runtime/UMG/Private/UserWidget.cpp
@@ -251,7 +251,7 @@ UWorld* UUserWidget::GetWorld() const
 	return nullptr;
 }
 
-UUMGSequencePlayer* UUserWidget::GetOrAddPlayer(const UWidgetAnimation* InAnimation)
+UUMGSequencePlayer* UUserWidget::GetOrAddPlayer(UWidgetAnimation* InAnimation)
 {
 	if (InAnimation)
 	{
@@ -291,7 +291,7 @@ void UUserWidget::Invalidate()
 	}
 }
 
-void UUserWidget::PlayAnimation( const UWidgetAnimation* InAnimation, float StartAtTime, int32 NumberOfLoops, EUMGSequencePlayMode::Type PlayMode, float PlaybackSpeed)
+void UUserWidget::PlayAnimation( UWidgetAnimation* InAnimation, float StartAtTime, int32 NumberOfLoops, EUMGSequencePlayMode::Type PlayMode, float PlaybackSpeed)
 {
 	FScopedNamedEvent NamedEvent(FColor::Emerald, "Widget::PlayAnimation");
 
@@ -306,7 +306,7 @@ void UUserWidget::PlayAnimation( const UWidgetAnimation* InAnimation, float Star
 	}
 }
 
-void UUserWidget::PlayAnimationTo(const UWidgetAnimation* InAnimation, float StartAtTime, float EndAtTime, int32 NumberOfLoops, EUMGSequencePlayMode::Type PlayMode, float PlaybackSpeed)
+void UUserWidget::PlayAnimationTo(UWidgetAnimation* InAnimation, float StartAtTime, float EndAtTime, int32 NumberOfLoops, EUMGSequencePlayMode::Type PlayMode, float PlaybackSpeed)
 {
 	FScopedNamedEvent NamedEvent(FColor::Emerald, "Widget::PlayAnimationTo");
 
diff --git a/Engine/Source/Runtime/UMG/Public/Animation/MovieScene2DTransformSection.h b/Engine/Source/Runtime/UMG/Public/Animation/MovieScene2DTransformSection.h
--- a/Engine/Source/Runtime/UMG/Public/Animation/MovieScene2DTransformSection.h
+++ b/Engine/Source/Runtime/UMG/Public/Animation/MovieScene2DTransformSection.h
@@ -60,13 +60,21 @@ public:
 public:
 
 	UMG_API FRichCurve& GetTranslationCurve( EAxis::Type Axis );
+	UMG_API const FRichCurve& GetTranslationCurve( EAxis::Type Axis ) const;
 
 	UMG_API FRichCurve& GetRotationCurve();
+	UMG_API const FRichCurve& GetRotationCurve() const;
 
 	UMG_API FRichCurve& GetScaleCurve( EAxis::Type Axis );
+	UMG_API const FRichCurve& GetScaleCurve( EAxis::Type Axis ) const;
 
+	DEPRECATED(4.15, "Please use GetShearCurve.")
 	UMG_API FRichCurve& GetSheerCurve( EAxis::Type Axis );
+	
+	UMG_API FRichCurve& GetShearCurve( EAxis::Type Axis );
+	UMG_API const FRichCurve& GetShearCurve( EAxis::Type Axis ) const;
 
+	DEPRECATED(4.15, "Evaluation is now the responsibility of FMovieScene2DTransformSectionTemplate")
 	FWidgetTransform Eval( float Position, const FWidgetTransform& DefaultValue ) const;
 
 	// IKeyframeSection interface.
diff --git a/Engine/Source/Runtime/UMG/Public/Animation/MovieScene2DTransformTrack.h b/Engine/Source/Runtime/UMG/Public/Animation/MovieScene2DTransformTrack.h
--- a/Engine/Source/Runtime/UMG/Public/Animation/MovieScene2DTransformTrack.h
+++ b/Engine/Source/Runtime/UMG/Public/Animation/MovieScene2DTransformTrack.h
@@ -23,17 +23,10 @@ public:
 	// UMovieSceneTrack interface
 
 	virtual UMovieSceneSection* CreateNewSection() override;
-	virtual TSharedPtr<IMovieSceneTrackInstance> CreateInstance() override;
+	virtual FMovieSceneEvalTemplatePtr CreateTemplateForSection(const UMovieSceneSection& InSection) const override;
 	
 public:
 
-	/**
-	 * Evaluates the track at the playback position
-	 *
-	 * @param Position	The current playback position
-	 * @param LastPosition	The last playback position
-	 * @param InOutTransform 	The transform at the playback position
-	 * @return true if anything was evaluated. Note: if false is returned InOutTransform remains unchanged
-	 */
+	DEPRECATED(4.15, "Please evaluate using FMovieScene2DTransformTemplate.")
 	bool Eval( float Position, float LastPostion, FWidgetTransform& InOutTransform ) const;
 };
diff --git a/Engine/Source/Runtime/UMG/Public/Animation/MovieSceneMarginSection.h b/Engine/Source/Runtime/UMG/Public/Animation/MovieSceneMarginSection.h
--- a/Engine/Source/Runtime/UMG/Public/Animation/MovieSceneMarginSection.h
+++ b/Engine/Source/Runtime/UMG/Public/Animation/MovieSceneMarginSection.h
@@ -4,6 +4,7 @@
 
 #include "MovieSceneSection.h"
 #include "IKeyframeSection.h"
+#include "Evaluation/MovieSceneEvalTemplate.h"
 #include "MovieSceneMarginSection.generated.h"
 
 
@@ -45,13 +46,6 @@ public:
 	virtual TOptional<float> GetKeyTime( FKeyHandle KeyHandle ) const override;
 	virtual void SetKeyTime( FKeyHandle KeyHandle, float Time ) override;
 
-	/**
-	 * Updates this section
-	 *
-	 * @param Position	The position in time within the movie scene
-	 */
-	FMargin Eval( float Position, const FMargin& DefaultValue ) const;
-
 	// IKeyframeSection interface.
 	virtual void AddKey( float Time, const FMarginKey& MarginKey, EMovieSceneKeyInterpolation KeyInterpolation ) override;
 	virtual bool NewKeyIsNewData(float Time, const FMarginKey& Key ) const override;
diff --git a/Engine/Source/Runtime/UMG/Public/Animation/MovieSceneMarginTrack.h b/Engine/Source/Runtime/UMG/Public/Animation/MovieSceneMarginTrack.h
--- a/Engine/Source/Runtime/UMG/Public/Animation/MovieSceneMarginTrack.h
+++ b/Engine/Source/Runtime/UMG/Public/Animation/MovieSceneMarginTrack.h
@@ -17,18 +17,10 @@ class UMovieSceneMarginTrack : public UMovieScenePropertyTrack
 
 public:
 
+	UMovieSceneMarginTrack(const FObjectInitializer& Init);
+	
 	// UMovieSceneTrack interface
 
 	virtual UMovieSceneSection* CreateNewSection() override;
-	virtual TSharedPtr<IMovieSceneTrackInstance> CreateInstance() override;
-
-	/**
-	 * Evaluates the track at the playback position
-	 *
-	 * @param Position The current playback position
-	 * @param LastPosition The last playback position
-	 * @param InOutMargin The margin at the playback position
-	 * @return true if anything was evaluated. Note: if false is returned OutMargin remains unchanged
-	 */
-	bool Eval( float Position, float LastPostion, FMargin& InOutMargin ) const;
+	virtual FMovieSceneEvalTemplatePtr CreateTemplateForSection(const UMovieSceneSection& InSection) const override;
 };
diff --git a/Engine/Source/Runtime/UMG/Public/Animation/UMGSequencePlayer.h b/Engine/Source/Runtime/UMG/Public/Animation/UMGSequencePlayer.h
--- a/Engine/Source/Runtime/UMG/Public/Animation/UMGSequencePlayer.h
+++ b/Engine/Source/Runtime/UMG/Public/Animation/UMGSequencePlayer.h
@@ -3,6 +3,7 @@
 #pragma once
 
 #include "IMovieScenePlayer.h"
+#include "Evaluation/MovieSceneEvaluationTemplateInstance.h"
 #include "UMGSequencePlayer.generated.h"
 
 class UWidgetAnimation;
@@ -14,7 +15,7 @@ class UMG_API UUMGSequencePlayer : public UObject, public IMovieScenePlayer
 	GENERATED_UCLASS_BODY()
 
 public:
-	void InitSequencePlayer(const UWidgetAnimation& InAnimation, UUserWidget& InUserWidget);
+	void InitSequencePlayer(UWidgetAnimation& InAnimation, UUserWidget& InUserWidget);
 
 	/** Updates the running movie */
 	void Tick( float DeltaTime );
@@ -47,13 +48,10 @@ public:
 	void SetPlaybackSpeed(float PlaybackSpeed);
 
 	/** IMovieScenePlayer interface */
-	virtual void GetRuntimeObjects( TSharedRef<FMovieSceneSequenceInstance> MovieSceneInstance, const FGuid& ObjectHandle, TArray<TWeakObjectPtr<UObject>>& OutObjects ) const override;
+	virtual FMovieSceneRootEvaluationTemplateInstance& GetEvaluationTemplate() override { return RootTemplateInstance; }
 	virtual void UpdateCameraCut(UObject* CameraObject, UObject* UnlockIfCameraObject, bool bJumpCut) override {}
 	virtual void SetViewportSettings(const TMap<FViewportClient*, EMovieSceneViewportParams>& ViewportParamsMap) override {}
 	virtual void GetViewportSettings(TMap<FViewportClient*, EMovieSceneViewportParams>& ViewportParamsMap) const override {}
-	virtual void AddOrUpdateMovieSceneInstance( class UMovieSceneSection& MovieSceneSection, TSharedRef<FMovieSceneSequenceInstance> InstanceToAdd ) override {}
-	virtual void RemoveMovieSceneInstance( class UMovieSceneSection& MovieSceneSection, TSharedRef<FMovieSceneSequenceInstance> InstanceToRemove ) override {}
-	virtual TSharedRef<FMovieSceneSequenceInstance> GetRootMovieSceneSequenceInstance() const override { return RootMovieSceneInstance.ToSharedRef(); }
 	virtual EMovieScenePlayerStatus::Type GetPlaybackStatus() const override;
 	virtual UObject* GetPlaybackContext() const override;
 	virtual TArray<UObject*> GetEventContexts() const override;
@@ -68,15 +66,13 @@ private:
 
 	/** Animation being played */
 	UPROPERTY()
-	const UWidgetAnimation* Animation;
+	UWidgetAnimation* Animation;
 
 	/** The user widget this sequence is animating */
 	TWeakObjectPtr<UUserWidget> UserWidget;
 
-	TMap<FGuid, TArray<UObject*> > GuidToRuntimeObjectMap;
-
-	/** The root movie scene instance to update when playing. */
-	TSharedPtr<class FMovieSceneSequenceInstance> RootMovieSceneInstance;
+	UPROPERTY(transient)
+	FMovieSceneRootEvaluationTemplateInstance RootTemplateInstance;
 
 	/** Time range of the animation */
 	TRange<float> TimeRange;
diff --git a/Engine/Source/Runtime/UMG/Public/Animation/WidgetAnimation.h b/Engine/Source/Runtime/UMG/Public/Animation/WidgetAnimation.h
--- a/Engine/Source/Runtime/UMG/Public/Animation/WidgetAnimation.h
+++ b/Engine/Source/Runtime/UMG/Public/Animation/WidgetAnimation.h
@@ -54,40 +54,20 @@ public:
 	UPROPERTY(BlueprintAssignable, Category="Animation")
 	FOnWidgetAnimationPlaybackStatusChanged OnAnimationFinished;
 
-	/**
-	 * Initialize the animation with a new user widget.
-	 *
-	 * @param InPreviewWidget The user widget to preview.
-	 */
-	UMG_API void Initialize(UUserWidget* InPreviewWidget);
-
 public:
 
 	// UMovieSceneAnimation overrides
 
 	virtual void BindPossessableObject(const FGuid& ObjectId, UObject& PossessedObject, UObject* Context) override;
-	virtual bool CanPossessObject(UObject& Object) const override;
-	virtual UObject* FindPossessableObject(const FGuid& ObjectId, UObject* Context) const override;
-	virtual FGuid FindPossessableObjectId(UObject& Object) const override;
+	virtual bool CanPossessObject(UObject& Object, UObject* InPlaybackContext) const override;
 	virtual UMovieScene* GetMovieScene() const override;
 	virtual UObject* GetParentObject(UObject* Object) const override;
 	virtual void UnbindPossessableObjects(const FGuid& ObjectId) override;
-
-	// UWidgetAnimation specific
-	/**
-	 * Replace the PossessedObject with a NewObject
-	 * 
-	 * @param OldId the Guid of the object is about to be replaced
-	 * @param NewId the Guid of the new object
-	 * @param OldObject the object is about to be replaced
-	 * @param NewObject the new object to replace
-	 */
-	UMG_API void ReplacePossessableObject(const FGuid& OldId, const FGuid& NewId, UObject& OldObject, UObject& NewObject);
+	virtual void LocateBoundObjects(const FGuid& ObjectId, UObject* Context, TArray<UObject*, TInlineAllocator<1>>& OutObjects) const override;
 
 	/** Get Animation bindings of the animation */
 	const TArray<FWidgetAnimationBinding>& GetBindings() const { return AnimationBindings; }
-	/** Get the preview widget of the animation, if any */
-	const TWeakObjectPtr<UUserWidget> GetPreviewWidget() const { return PreviewWidget; }
+
 public:
 
 	/** Pointer to the movie scene that controls this animation. */
@@ -96,15 +76,4 @@ public:
 
 	UPROPERTY()
 	TArray<FWidgetAnimationBinding> AnimationBindings;
-
-private:
-
-	/** The current preview widget, if any. */
-	TWeakObjectPtr<UUserWidget> PreviewWidget;
-
-	/** Mapping of preview objects to sequencer GUIDs */
-	TMultiMap<FGuid, TWeakObjectPtr<UObject>> IdToPreviewObjects;
-	TMultiMap<FGuid, TWeakObjectPtr<UObject>> IdToSlotContentPreviewObjects;
-	TMap<TWeakObjectPtr<UObject>, FGuid> PreviewObjectToIds;
-	TMap<TWeakObjectPtr<UObject>, FGuid> SlotContentPreviewObjectToIds;
 };
diff --git a/Engine/Source/Runtime/UMG/Public/Blueprint/UserWidget.h b/Engine/Source/Runtime/UMG/Public/Blueprint/UserWidget.h
--- a/Engine/Source/Runtime/UMG/Public/Blueprint/UserWidget.h
+++ b/Engine/Source/Runtime/UMG/Public/Blueprint/UserWidget.h
@@ -710,7 +710,7 @@ public:
 	 * @param PlayMode Specifies the playback mode
 	 */
 	UFUNCTION(BlueprintCallable, BlueprintCosmetic, Category = "User Interface|Animation")
-	void PlayAnimation(const UWidgetAnimation* InAnimation, float StartAtTime = 0.0f, int32 NumLoopsToPlay = 1, EUMGSequencePlayMode::Type PlayMode = EUMGSequencePlayMode::Forward, float PlaybackSpeed = 1.0f);
+	void PlayAnimation(UWidgetAnimation* InAnimation, float StartAtTime = 0.0f, int32 NumLoopsToPlay = 1, EUMGSequencePlayMode::Type PlayMode = EUMGSequencePlayMode::Forward, float PlaybackSpeed = 1.0f);
 
 	/**
 	 * Plays an animation in this widget a specified number of times stoping at a specified time
@@ -723,7 +723,7 @@ public:
 	 * @param PlayMode Specifies the playback mode
 	 */
 	UFUNCTION(BlueprintCallable, BlueprintCosmetic, Category="User Interface|Animation")
-	void PlayAnimationTo(const UWidgetAnimation* InAnimation, float StartAtTime = 0.0f, float EndAtTime = 0.0f, int32 NumLoopsToPlay = 1, EUMGSequencePlayMode::Type PlayMode = EUMGSequencePlayMode::Forward, float PlaybackSpeed = 1.0f);
+	void PlayAnimationTo(UWidgetAnimation* InAnimation, float StartAtTime = 0.0f, float EndAtTime = 0.0f, int32 NumLoopsToPlay = 1, EUMGSequencePlayMode::Type PlayMode = EUMGSequencePlayMode::Forward, float PlaybackSpeed = 1.0f);
 
 	/**
 	 * Stops an already running animation in this widget
@@ -989,7 +989,7 @@ protected:
 
 	void RemoveObsoleteBindings(const TArray<FName>& NamedSlots);
 
-	UUMGSequencePlayer* GetOrAddPlayer(const UWidgetAnimation* InAnimation);
+	UUMGSequencePlayer* GetOrAddPlayer(UWidgetAnimation* InAnimation);
 	void Invalidate();
 	
 	/**
diff --git a/Engine/Source/Runtime/UtilityShaders/Private/MediaShaders.cpp b/Engine/Source/Runtime/UtilityShaders/Private/MediaShaders.cpp
--- a/Engine/Source/Runtime/UtilityShaders/Private/MediaShaders.cpp
+++ b/Engine/Source/Runtime/UtilityShaders/Private/MediaShaders.cpp
@@ -39,26 +39,57 @@ void FAYUVConvertPS::SetParameters(FRHICommandList& RHICmdList, TRefCountPtr<FRH
 }
 
 
+/* FRGBConvertPS shader
+ *****************************************************************************/
+
+BEGIN_UNIFORM_BUFFER_STRUCT(FBMPConvertUB, )
+DECLARE_UNIFORM_BUFFER_STRUCT_MEMBER(FVector2D, UVScale)
+DECLARE_UNIFORM_BUFFER_STRUCT_MEMBER_TEXTURE(Texture2D, BMPTexture)
+DECLARE_UNIFORM_BUFFER_STRUCT_MEMBER_SAMPLER(SamplerState, BMPSampler)
+END_UNIFORM_BUFFER_STRUCT(FBMPConvertUB)
+
+IMPLEMENT_UNIFORM_BUFFER_STRUCT(FBMPConvertUB, TEXT("BMPConvertUB"));
+IMPLEMENT_SHADER_TYPE(, FBMPConvertPS, TEXT("MediaShaders"), TEXT("BMPConvertPS"), SF_Pixel);
+
+
+void FBMPConvertPS::SetParameters(FRHICommandList& RHICmdList, TRefCountPtr<FRHITexture2D> BMPTexture, const FIntPoint& OutputDimensions)
+{
+	FBMPConvertUB UB;
+	{
+		UB.UVScale = FVector2D((float)OutputDimensions.X / (float)BMPTexture->GetSizeX(), (float)OutputDimensions.Y / (float)BMPTexture->GetSizeY());
+		UB.BMPTexture = BMPTexture;
+		UB.BMPSampler = TStaticSamplerState<SF_Bilinear>::GetRHI();
+	}
+
+	TUniformBufferRef<FBMPConvertUB> Data = TUniformBufferRef<FBMPConvertUB>::CreateUniformBufferImmediate(UB, UniformBuffer_SingleFrame);
+	SetUniformBufferParameter(RHICmdList, GetPixelShader(), GetUniformBufferParameter<FBMPConvertUB>(), Data);
+}
+
+
 /* FNV12ConvertPS shader
  *****************************************************************************/
 
 BEGIN_UNIFORM_BUFFER_STRUCT(FNV12ConvertUB, )
-DECLARE_UNIFORM_BUFFER_STRUCT_MEMBER(uint32, Width)
+DECLARE_UNIFORM_BUFFER_STRUCT_MEMBER(uint32, OutputWidth)
+DECLARE_UNIFORM_BUFFER_STRUCT_MEMBER(FVector2D, UVScale)
 DECLARE_UNIFORM_BUFFER_STRUCT_MEMBER_TEXTURE(Texture2D, NV12Texture)
-DECLARE_UNIFORM_BUFFER_STRUCT_MEMBER_SAMPLER(SamplerState, NV12Sampler)
+DECLARE_UNIFORM_BUFFER_STRUCT_MEMBER_SAMPLER(SamplerState, NV12SamplerB)
+DECLARE_UNIFORM_BUFFER_STRUCT_MEMBER_SAMPLER(SamplerState, NV12SamplerP)
 END_UNIFORM_BUFFER_STRUCT(FNV12ConvertUB)
 
 IMPLEMENT_UNIFORM_BUFFER_STRUCT(FNV12ConvertUB, TEXT("NV12ConvertUB"));
 IMPLEMENT_SHADER_TYPE(, FNV12ConvertPS, TEXT("MediaShaders"), TEXT("NV12ConvertPS"), SF_Pixel);
 
 
-void FNV12ConvertPS::SetParameters(FRHICommandList& RHICmdList, TRefCountPtr<FRHITexture2D> NV12Texture)
+void FNV12ConvertPS::SetParameters(FRHICommandList& RHICmdList, TRefCountPtr<FRHITexture2D> NV12Texture, const FIntPoint& OutputDimensions)
 {
 	FNV12ConvertUB UB;
 	{
-		UB.Width = NV12Texture->GetSizeX();
+		UB.OutputWidth = OutputDimensions.X;
+		UB.UVScale = FVector2D((float)OutputDimensions.X / (float)NV12Texture->GetSizeX(), (float)OutputDimensions.Y / (float)NV12Texture->GetSizeY());
 		UB.NV12Texture = NV12Texture;
-		UB.NV12Sampler = TStaticSamplerState<SF_Bilinear>::GetRHI();
+		UB.NV12SamplerB = TStaticSamplerState<SF_Bilinear>::GetRHI();
+		UB.NV12SamplerP = TStaticSamplerState<SF_Point>::GetRHI();
 	}
 
 	TUniformBufferRef<FNV12ConvertUB> Data = TUniformBufferRef<FNV12ConvertUB>::CreateUniformBufferImmediate(UB, UniformBuffer_SingleFrame);
@@ -70,22 +101,26 @@ void FNV12ConvertPS::SetParameters(FRHICommandList& RHICmdList, TRefCountPtr<FRH
  *****************************************************************************/
 
 BEGIN_UNIFORM_BUFFER_STRUCT(FNV21ConvertUB, )
-DECLARE_UNIFORM_BUFFER_STRUCT_MEMBER(uint32, Width)
+DECLARE_UNIFORM_BUFFER_STRUCT_MEMBER(uint32, OutputWidth)
+DECLARE_UNIFORM_BUFFER_STRUCT_MEMBER(FVector2D, UVScale)
 DECLARE_UNIFORM_BUFFER_STRUCT_MEMBER_TEXTURE(Texture2D, NV21Texture)
-DECLARE_UNIFORM_BUFFER_STRUCT_MEMBER_SAMPLER(SamplerState, NV21Sampler)
+DECLARE_UNIFORM_BUFFER_STRUCT_MEMBER_SAMPLER(SamplerState, NV21SamplerB)
+DECLARE_UNIFORM_BUFFER_STRUCT_MEMBER_SAMPLER(SamplerState, NV21SamplerP)
 END_UNIFORM_BUFFER_STRUCT(FNV21ConvertUB)
 
 IMPLEMENT_UNIFORM_BUFFER_STRUCT(FNV21ConvertUB, TEXT("NV21ConvertUB"));
 IMPLEMENT_SHADER_TYPE(, FNV21ConvertPS, TEXT("MediaShaders"), TEXT("NV21ConvertPS"), SF_Pixel);
 
 
-void FNV21ConvertPS::SetParameters(FRHICommandList& RHICmdList, TRefCountPtr<FRHITexture2D> NV21Texture)
+void FNV21ConvertPS::SetParameters(FRHICommandList& RHICmdList, TRefCountPtr<FRHITexture2D> NV21Texture, const FIntPoint& OutputDimensions)
 {
 	FNV21ConvertUB UB;
 	{
-		UB.Width = NV21Texture->GetSizeX();
+		UB.OutputWidth = OutputDimensions.Y;
+		UB.UVScale = FVector2D((float)OutputDimensions.X / (float)NV21Texture->GetSizeX(), (float)OutputDimensions.Y / (float)NV21Texture->GetSizeY());
 		UB.NV21Texture = NV21Texture;
-		UB.NV21Sampler = TStaticSamplerState<SF_Bilinear>::GetRHI();
+		UB.NV21SamplerB = TStaticSamplerState<SF_Bilinear>::GetRHI();
+		UB.NV21SamplerP = TStaticSamplerState<SF_Point>::GetRHI();
 	}
 
 	TUniformBufferRef<FNV21ConvertUB> Data = TUniformBufferRef<FNV21ConvertUB>::CreateUniformBufferImmediate(UB, UniformBuffer_SingleFrame);
@@ -93,6 +128,33 @@ void FNV21ConvertPS::SetParameters(FRHICommandList& RHICmdList, TRefCountPtr<FRH
 }
 
 
+/* FRGBConvertPS shader
+ *****************************************************************************/
+
+BEGIN_UNIFORM_BUFFER_STRUCT(FRGBConvertUB, )
+DECLARE_UNIFORM_BUFFER_STRUCT_MEMBER(FVector2D, UVScale)
+DECLARE_UNIFORM_BUFFER_STRUCT_MEMBER_TEXTURE(Texture2D, RGBTexture)
+DECLARE_UNIFORM_BUFFER_STRUCT_MEMBER_SAMPLER(SamplerState, RGBSampler)
+END_UNIFORM_BUFFER_STRUCT(FRGBConvertUB)
+
+IMPLEMENT_UNIFORM_BUFFER_STRUCT(FRGBConvertUB, TEXT("RGBConvertUB"));
+IMPLEMENT_SHADER_TYPE(, FRGBConvertPS, TEXT("MediaShaders"), TEXT("RGBConvertPS"), SF_Pixel);
+
+
+void FRGBConvertPS::SetParameters(FRHICommandList& RHICmdList, TRefCountPtr<FRHITexture2D> RGBTexture, const FIntPoint& OutputDimensions)
+{
+	FRGBConvertUB UB;
+	{
+		UB.UVScale = FVector2D((float)OutputDimensions.X / (float)RGBTexture->GetSizeX(), (float)OutputDimensions.Y / (float)RGBTexture->GetSizeY());
+		UB.RGBTexture = RGBTexture;
+		UB.RGBSampler = TStaticSamplerState<SF_Bilinear>::GetRHI();
+	}
+
+	TUniformBufferRef<FRGBConvertUB> Data = TUniformBufferRef<FRGBConvertUB>::CreateUniformBufferImmediate(UB, UniformBuffer_SingleFrame);
+	SetUniformBufferParameter(RHICmdList, GetPixelShader(), GetUniformBufferParameter<FRGBConvertUB>(), Data);
+}
+
+
 /* FYCbCrConvertUB shader
  *****************************************************************************/
 
@@ -188,7 +250,8 @@ void FYUVConvertPS::SetParameters(FRHICommandList& RHICmdList, TRefCountPtr<FRHI
  *****************************************************************************/
 
 BEGIN_UNIFORM_BUFFER_STRUCT(FYUY2ConvertUB, )
-DECLARE_UNIFORM_BUFFER_STRUCT_MEMBER(uint32, Width)
+DECLARE_UNIFORM_BUFFER_STRUCT_MEMBER(uint32, OutputWidth)
+DECLARE_UNIFORM_BUFFER_STRUCT_MEMBER(FVector2D, UVScale)
 DECLARE_UNIFORM_BUFFER_STRUCT_MEMBER_TEXTURE(Texture2D, YUY2Texture)
 DECLARE_UNIFORM_BUFFER_STRUCT_MEMBER_SAMPLER(SamplerState, YUY2SamplerB)
 DECLARE_UNIFORM_BUFFER_STRUCT_MEMBER_SAMPLER(SamplerState, YUY2SamplerP)
@@ -198,11 +261,12 @@ IMPLEMENT_UNIFORM_BUFFER_STRUCT(FYUY2ConvertUB, TEXT("YUY2ConvertUB"));
 IMPLEMENT_SHADER_TYPE(, FYUY2ConvertPS, TEXT("MediaShaders"), TEXT("YUY2ConvertPS"), SF_Pixel);
 
 
-void FYUY2ConvertPS::SetParameters(FRHICommandList& RHICmdList, TRefCountPtr<FRHITexture2D> YUY2Texture)
+void FYUY2ConvertPS::SetParameters(FRHICommandList& RHICmdList, TRefCountPtr<FRHITexture2D> YUY2Texture, const FIntPoint& OutputDimensions)
 {
 	FYUY2ConvertUB UB;
 	{
-		UB.Width = YUY2Texture->GetSizeX();
+		UB.OutputWidth = OutputDimensions.X;
+		UB.UVScale = FVector2D((float)OutputDimensions.X / (2.0f * YUY2Texture->GetSizeX()), (float)OutputDimensions.Y / (float)YUY2Texture->GetSizeY());
 		UB.YUY2Texture = YUY2Texture;
 		UB.YUY2SamplerB = TStaticSamplerState<SF_Bilinear>::GetRHI();
 		UB.YUY2SamplerP = TStaticSamplerState<SF_Point>::GetRHI();
diff --git a/Engine/Source/Runtime/UtilityShaders/Public/MediaShaders.h b/Engine/Source/Runtime/UtilityShaders/Public/MediaShaders.h
--- a/Engine/Source/Runtime/UtilityShaders/Public/MediaShaders.h
+++ b/Engine/Source/Runtime/UtilityShaders/Public/MediaShaders.h
@@ -121,10 +121,43 @@ public:
 };
 
 
+/**
+ * Pixel shader to convert a Windows Bitmap texture.
+ *
+ * This shader expects a BMP frame packed into a single texture in PF_B8G8R8A8 format.
+ */
+class FBMPConvertPS
+	: public FGlobalShader
+{
+	DECLARE_EXPORTED_SHADER_TYPE(FBMPConvertPS, Global, UTILITYSHADERS_API);
+
+public:
+
+	static bool ShouldCache(EShaderPlatform Platform)
+	{
+		return IsFeatureLevelSupported(Platform, ERHIFeatureLevel::ES2);
+	}
+
+	FBMPConvertPS() { }
+
+	FBMPConvertPS(const ShaderMetaType::CompiledShaderInitializerType& Initializer)
+		: FGlobalShader(Initializer)
+	{ }
+
+	virtual bool Serialize(FArchive& Ar) override
+	{
+		bool bShaderHasOutdatedParameters = FGlobalShader::Serialize(Ar);
+		return bShaderHasOutdatedParameters;
+	}
+
+	UTILITYSHADERS_API void SetParameters(FRHICommandList& RHICmdList, TRefCountPtr<FRHITexture2D> BMPTexture, const FIntPoint& OutputDimensions);
+};
+
+
 /**
  * Pixel shader to convert a NV12 frame to RGBA.
  *
- * This shader expects an NV12 frame packed into single texture in PF_G8 format.
+ * This shader expects an NV12 frame packed into a single texture in PF_G8 format.
  *
  * @see http://www.fourcc.org/yuv.php#NV12
  */
@@ -152,14 +185,14 @@ public:
 		return bShaderHasOutdatedParameters;
 	}
 
-	UTILITYSHADERS_API void SetParameters(FRHICommandList& RHICmdList, TRefCountPtr<FRHITexture2D> NV12Texture);
+	UTILITYSHADERS_API void SetParameters(FRHICommandList& RHICmdList, TRefCountPtr<FRHITexture2D> NV12Texture, const FIntPoint& OutputDimensions);
 };
 
 
 /**
  * Pixel shader to convert a NV21 frame to RGBA.
  *
- * This shader expects an NV21 frame packed into single texture in PF_G8 format.
+ * This shader expects an NV21 frame packed into a single texture in PF_G8 format.
  *
  * @see http://www.fourcc.org/yuv.php#NV21
  */
@@ -187,7 +220,41 @@ public:
 		return bShaderHasOutdatedParameters;
 	}
 
-	UTILITYSHADERS_API void SetParameters(FRHICommandList& RHICmdList, TRefCountPtr<FRHITexture2D> NV21Texture);
+	UTILITYSHADERS_API void SetParameters(FRHICommandList& RHICmdList, TRefCountPtr<FRHITexture2D> NV21Texture, const FIntPoint& OutputDimensions);
+};
+
+
+/**
+ * Pixel shader to resize an RGB texture.
+ *
+ * This shader expects an RGB or RGBA frame packed into a single texture
+ * in PF_B8G8R8A8 or PF_FloatRGB format.
+ */
+class FRGBConvertPS
+	: public FGlobalShader
+{
+	DECLARE_EXPORTED_SHADER_TYPE(FRGBConvertPS, Global, UTILITYSHADERS_API);
+
+public:
+
+	static bool ShouldCache(EShaderPlatform Platform)
+	{
+		return IsFeatureLevelSupported(Platform, ERHIFeatureLevel::ES2);
+	}
+
+	FRGBConvertPS() { }
+
+	FRGBConvertPS(const ShaderMetaType::CompiledShaderInitializerType& Initializer)
+		: FGlobalShader(Initializer)
+	{ }
+
+	virtual bool Serialize(FArchive& Ar) override
+	{
+		bool bShaderHasOutdatedParameters = FGlobalShader::Serialize(Ar);
+		return bShaderHasOutdatedParameters;
+	}
+
+	UTILITYSHADERS_API void SetParameters(FRHICommandList& RHICmdList, TRefCountPtr<FRHITexture2D> RGBTexture, const FIntPoint& OutputDimensions);
 };
 
 
@@ -301,7 +368,7 @@ public:
 /**
  * Pixel shader to convert a YUY2 frame to RGBA.
  *
- * This shader expects an YUY2 frame packed into single texture in PF_B8G8R8A8
+ * This shader expects an YUY2 frame packed into a single texture in PF_B8G8R8A8
  * format with the following memory layout: [Y0, U0, Y1, V0][Y2, U1, Y3, V1]...
  *
  * @see http://www.fourcc.org/yuv.php#YUY2
@@ -330,7 +397,7 @@ public:
 		return bShaderHasOutdatedParameters;
 	}
 
-	UTILITYSHADERS_API void SetParameters(FRHICommandList& RHICmdList, TRefCountPtr<FRHITexture2D> YUY2Texture);
+	UTILITYSHADERS_API void SetParameters(FRHICommandList& RHICmdList, TRefCountPtr<FRHITexture2D> YUY2Texture, const FIntPoint& OutputDimensions);
 };
 
 
