Fix Online compile issues.

diff --git a/Engine/Source/Runtime/Sockets/Private/BSDIPv6Sockets/IPAddressBSDIPv6.h b/Engine/Source/Runtime/Sockets/Private/BSDIPv6Sockets/IPAddressBSDIPv6.h
--- a/Engine/Source/Runtime/Sockets/Private/BSDIPv6Sockets/IPAddressBSDIPv6.h
+++ b/Engine/Source/Runtime/Sockets/Private/BSDIPv6Sockets/IPAddressBSDIPv6.h
@@ -307,18 +307,12 @@ public:
 			Addr.sin6_family == OtherBSD.Addr.sin6_family;
 	}
 
-	virtual bool operator<(const FInternetAddr& Other) const override
-	{
-		FInternetAddrBSDIPv6& OtherBSD = (FInternetAddrBSDIPv6&)Other;
-
-		return Addr.sin6_family < OtherBSD.Addr.sin6_family ||
-				FMemory::Memcmp(Addr.sin6_addr.u.Bytes, OtherBSD.Addr.sin6_addr.u.Bytes) < 0 ||
-				Addr.sin6_port < OtherBSD.Addr.sin6_port;
-	}
-
 	virtual uint32 GetTypeHash() override
 	{
-		return GetTypeHash(Addr.sin_addr.s_addr) + (Addr.sin_port * 23);
+		// @todo: Find a more efficient way to hash IPv6 addresses, if they are to be used with NetConnection's
+		//return Addr.sin_addr.s_addr + (Addr.sin_port * 23);
+
+		return ::GetTypeHash(*ToString(true));
 	}
 
 	/**
diff --git a/Engine/Source/Runtime/Sockets/Private/BSDSockets/IPAddressBSD.h b/Engine/Source/Runtime/Sockets/Private/BSDSockets/IPAddressBSD.h
--- a/Engine/Source/Runtime/Sockets/Private/BSDSockets/IPAddressBSD.h
+++ b/Engine/Source/Runtime/Sockets/Private/BSDSockets/IPAddressBSD.h
@@ -193,14 +193,6 @@ public:
 			Addr.sin_family == OtherBSD.Addr.sin_family;
 	}
 
-	virtual bool operator<(const FInternetAddr& Other) const override
-	{
-		FInternetAddrBSD& OtherBSD = (FInternetAddrBSD&)Other;
-
-		return Addr.sin_family < OtherBSD.Addr.sin_family || ntohl(Addr.sin_addr.s_addr) < ntohl(OtherBSD.Addr.sin_addr.s_addr) ||
-				ntohs(Addr.sin_port) < ntohs(OtherBSD.Addr.sin_port);
-	}
-
 	virtual uint32 GetTypeHash() override
 	{
 		return Addr.sin_addr.s_addr + (Addr.sin_port * 23);
diff --git a/Engine/Source/Runtime/Sockets/Public/IPAddress.h b/Engine/Source/Runtime/Sockets/Public/IPAddress.h
--- a/Engine/Source/Runtime/Sockets/Public/IPAddress.h
+++ b/Engine/Source/Runtime/Sockets/Public/IPAddress.h
@@ -105,21 +105,6 @@ public:
 		return ThisIP == OtherIP && GetPort() == Other.GetPort();
 	}
 
-	/**
-	 * Compares two internet addresses for sorting (used with binary searching IP lists)
-	 *
-	 * @param Other		The other address to compare against
-	 */
-	virtual bool operator<(const FInternetAddr& Other) const
-	{
-		uint32 ThisIP, OtherIP;
-
-		GetIp(ThisIP);
-		Other.GetIp(OtherIP);
-
-		return (ThisIP < OtherIP) || GetPort() < Other.GetPort();
-	}
-
 	/**
 	 * Hash function for use with TMap's - exposed through FInternetAddrMapRef
 	 *
