Merging IPv6 and IPv4 code together. Exposing a getaddrinfo method.

diff --git a/Engine/Plugins/Online/OnlineSubsystem/Source/Public/NboSerializer.h b/Engine/Plugins/Online/OnlineSubsystem/Source/Public/NboSerializer.h
--- a/Engine/Plugins/Online/OnlineSubsystem/Source/Public/NboSerializer.h
+++ b/Engine/Plugins/Online/OnlineSubsystem/Source/Public/NboSerializer.h
@@ -420,8 +420,7 @@ public:
 		uint32 OutIp;
 		Addr.GetIp(OutIp);
 		Ar << OutIp;
-		int32 OutPort;
-		Addr.GetPort(OutPort);
+		int32 OutPort = Addr.GetPort();
 		Ar << OutPort;
 		return Ar;
 	}
diff --git a/Engine/Plugins/Online/OnlineSubsystemNull/Source/Private/OnlineSessionInterfaceNull.cpp b/Engine/Plugins/Online/OnlineSubsystemNull/Source/Private/OnlineSessionInterfaceNull.cpp
--- a/Engine/Plugins/Online/OnlineSubsystemNull/Source/Private/OnlineSessionInterfaceNull.cpp
+++ b/Engine/Plugins/Online/OnlineSubsystemNull/Source/Private/OnlineSessionInterfaceNull.cpp
@@ -699,9 +699,7 @@ uint32 FOnlineSessionNull::JoinLANSession(int32 PlayerNum, FNamedOnlineSession*
 		FOnlineSessionInfoNull* SessionInfo = (FOnlineSessionInfoNull*)Session->SessionInfo.Get();
 		SessionInfo->SessionId = SearchSessionInfo->SessionId;
 
-		uint32 IpAddr;
-		SearchSessionInfo->HostAddr->GetIp(IpAddr);
-		SessionInfo->HostAddr = ISocketSubsystem::Get(PLATFORM_SOCKETSUBSYSTEM)->CreateInternetAddr(IpAddr, SearchSessionInfo->HostAddr->GetPort());
+		SessionInfo->HostAddr = SearchSessionInfo->HostAddr->Clone();
 		Result = ERROR_SUCCESS;
 	}
 
diff --git a/Engine/Plugins/Online/OnlineSubsystemSteam/Source/Private/IPAddressSteam.cpp b/Engine/Plugins/Online/OnlineSubsystemSteam/Source/Private/IPAddressSteam.cpp
--- a/Engine/Plugins/Online/OnlineSubsystemSteam/Source/Private/IPAddressSteam.cpp
+++ b/Engine/Plugins/Online/OnlineSubsystemSteam/Source/Private/IPAddressSteam.cpp
@@ -1,6 +1,46 @@
 // Copyright 1998-2018 Epic Games, Inc. All Rights Reserved.
 
 #include "IPAddressSteam.h"
+#include "Algo/Reverse.h"
+
+TArray<uint8> FInternetAddrSteam::GetRawIp() const
+{
+	TArray<uint8> RawAddressArray;
+	const uint8* SteamIdWalk = SteamId.GetBytes();
+	while (RawAddressArray.Num() < SteamId.GetSize())
+	{
+		RawAddressArray.Add(*SteamIdWalk);
+		++SteamIdWalk;
+	}
+
+	// We want to make sure that these arrays are in big endian format.
+#if PLATFORM_LITTLE_ENDIAN
+	Algo::Reverse(RawAddressArray);
+#endif
+
+	return RawAddressArray;
+}
+
+void FInternetAddrSteam::SetRawIp(const TArray<uint8>& RawAddr)
+{
+	uint64 NewSteamId = 0;
+
+	// Make a quick copy of the array
+	TArray<uint8> WorkingArray = RawAddr;
+
+	// Flip the entire array.
+	// Normally we would just do a ntohll on the final result but it's not a portable function.
+#if PLATFORM_LITTLE_ENDIAN
+	Algo::Reverse(WorkingArray);
+#endif
+
+	for (int32 i = 0; i < WorkingArray.Num(); ++i)
+	{
+		NewSteamId |= (uint64)WorkingArray[i] << (i * 8);
+	}
+
+	SteamId = FUniqueNetIdSteam(NewSteamId);
+}
 
 /**
  * Sets the ip address from a string ("steam.STEAMID" or "STEAMID")
diff --git a/Engine/Plugins/Online/OnlineSubsystemSteam/Source/Private/IPAddressSteam.h b/Engine/Plugins/Online/OnlineSubsystemSteam/Source/Private/IPAddressSteam.h
--- a/Engine/Plugins/Online/OnlineSubsystemSteam/Source/Private/IPAddressSteam.h
+++ b/Engine/Plugins/Online/OnlineSubsystemSteam/Source/Private/IPAddressSteam.h
@@ -45,6 +45,20 @@ public:
 	{
 	}
 
+	/**
+	 * Gets the ip address in a raw array stored in network byte order.
+	 *
+	 * @return The raw address stored in an array
+	 */
+	virtual TArray<uint8> GetRawIp() const override;
+
+	/**
+	 * Sets the ip address from a raw network byte order array.
+	 *
+	 * @param RawAddr the new address to use (must be converted to network byte order)
+	 */
+	virtual void SetRawIp(const TArray<uint8>& RawAddr) override;
+
 	/**
 	 * Sets the ip address from a host byte order uint32
 	 *
@@ -111,6 +125,11 @@ public:
 	{
 	}
 
+	/** Sets the address to the loopback address */
+	void SetLoopbackAddress() override
+	{
+	}
+
 	/**
 	 * Converts this internet ip address to string form
 	 *
@@ -158,4 +177,12 @@ public:
 	{
 		return SteamId.IsValid();
 	}
+
+	virtual TSharedRef<FInternetAddr> Clone() const override
+	{
+		TSharedRef<FInternetAddrSteam> NewAddress = MakeShareable(new FInternetAddrSteam);
+		NewAddress->SteamId = SteamId;
+		NewAddress->SteamChannel = SteamChannel;
+		return NewAddress;
+	}
 };
diff --git a/Engine/Plugins/Online/OnlineSubsystemUtils/Source/OnlineSubsystemUtils/Private/IpConnection.cpp b/Engine/Plugins/Online/OnlineSubsystemUtils/Source/OnlineSubsystemUtils/Private/IpConnection.cpp
--- a/Engine/Plugins/Online/OnlineSubsystemUtils/Source/OnlineSubsystemUtils/Private/IpConnection.cpp
+++ b/Engine/Plugins/Online/OnlineSubsystemUtils/Source/OnlineSubsystemUtils/Private/IpConnection.cpp
@@ -128,13 +128,7 @@ void UIpConnection::LowLevelSend(void* Data, int32 CountBytes, int32 CountBits)
 		else
 		{
 			// Host name resolution just now succeeded.
-#if PLATFORM_IOS
-			RemoteAddr->Copy(ResolveInfo->GetResolvedAddress());
-#else
-			uint32 Addr;
-			ResolveInfo->GetResolvedAddress().GetIp(Addr);
-			RemoteAddr->SetIp(Addr);
-#endif
+			RemoteAddr = ResolveInfo->GetResolvedAddress().Clone();
 			UE_LOG(LogNet, Log, TEXT("Host name resolution completed"));
 			delete ResolveInfo;
 			ResolveInfo = NULL;
@@ -246,10 +240,7 @@ int32 UIpConnection::GetAddrAsInt(void)
 
 int32 UIpConnection::GetAddrPort(void)
 {
-	int32 OutPort = 0;
-	// Get the host byte order ip port
-	RemoteAddr->GetPort(OutPort);
-	return OutPort;
+	return RemoteAddr->GetPort();
 }
 
 TSharedPtr<FInternetAddr> UIpConnection::GetInternetAddr()
diff --git a/Engine/Source/Runtime/Core/Public/HAL/Platform.h b/Engine/Source/Runtime/Core/Public/HAL/Platform.h
--- a/Engine/Source/Runtime/Core/Public/HAL/Platform.h
+++ b/Engine/Source/Runtime/Core/Public/HAL/Platform.h
@@ -306,6 +306,9 @@
 #ifndef PLATFORM_HAS_BSD_SOCKET_FEATURE_GETADDRINFO
 	#define PLATFORM_HAS_BSD_SOCKET_FEATURE_GETADDRINFO	1
 #endif
+#ifndef PLATFORM_HAS_BSD_SOCKET_FEATURE_GETNAMEINFO
+	#define PLATFORM_HAS_BSD_SOCKET_FEATURE_GETNAMEINFO 1
+#endif
 #ifndef PLATFORM_HAS_BSD_SOCKET_FEATURE_CLOSE_ON_EXEC
 	#define PLATFORM_HAS_BSD_SOCKET_FEATURE_CLOSE_ON_EXEC	0
 #endif
diff --git a/Engine/Source/Runtime/Sockets/Private/Android/SocketSubsystemAndroid.cpp b/Engine/Source/Runtime/Sockets/Private/Android/SocketSubsystemAndroid.cpp
--- a/Engine/Source/Runtime/Sockets/Private/Android/SocketSubsystemAndroid.cpp
+++ b/Engine/Source/Runtime/Sockets/Private/Android/SocketSubsystemAndroid.cpp
@@ -6,6 +6,7 @@
 #include "Misc/Parse.h"
 #include "Misc/CommandLine.h"
 #include "Modules/ModuleManager.h"
+#include "BSDSockets/IPAddressBSD.h"
 
 #include <sys/ioctl.h>
 #include <net/if.h>
@@ -111,33 +112,24 @@ TSharedRef<FInternetAddr> FSocketSubsystemAndroid::GetLocalHostAddr(FOutputDevic
 	// Depreciated function gethostname() returns 'localhost' on (all?) Android devices
 	// Which in turn means that FSocketSubsystemBSD::GetLocalHostAddr() resolves to 127.0.0.1
 	// Getting info from android.net.wifi.WifiManager is a little messy due to UE4 modular architecture and JNI
-	// IPv4 code using ioctl(.., SIOCGIFCONF, ..) actually based on FSocketSubsytemLinux::GetLocalHostAddr works fine for now...
+	// IPv4 code using ioctl(.., SIOCGIFCONF, ..) based on formally FSocketSubsytemLinux::GetLocalHostAddr works fine for now...
 	//
 	// Also NOTE: Network can flip out behind applications back when connectivity changes. eg. Move out of wifi range.
 	// This seems to recover OK between matches as subsystems are reinited each session Host/Join.
 
-	uint32 ParentIp = 0;
-	Addr->GetIp(ParentIp); // will return in host order
+	uint32 ParentIp;
+	Addr->GetIp(ParentIp);
 	if (ParentIp != 0 && (ParentIp & 0xff000000) != 0x7f000000)
 	{
 		return Addr;
 	}
 
-	// If superclass got the address from command line, honor that override
-	TCHAR Home[256] = TEXT("");
-	if (FParse::Value(FCommandLine::Get(), TEXT("MULTIHOME="), Home, ARRAY_COUNT(Home)))
-	{
-		TSharedRef<FInternetAddr> TempAddr = CreateInternetAddr();
-		bool bIsValid = false;
-		TempAddr->SetIp(Home, bIsValid);
-		if (bIsValid)
-		{
-			UE_LOG(LogSockets, Warning, TEXT("FSocketSubsystemAndroid::GetLocalHostAddr Using MULTIHOME"));
-			return TempAddr;
-		}
-	}
+	// TODO: Android doesn't support ifaddrs either except in the Android OS 7.0+
+	// Which isn't super great either. So I guess this block is stuck the way it is.
+	// Other alternatives could be rtnetdevice but it's blocking 
+	// (sure you can make it non-blocking but you still have to wait for a recv otherwise you have no data)
 
-	// We need to go deeper...
+	// we need to go deeper...  (see http://man7.org/linux/man-pages/man7/netdevice.7.html)
 	int TempSocket = socket(PF_INET, SOCK_STREAM, 0);
 	if (TempSocket)
 	{
@@ -151,9 +143,14 @@ TSharedRef<FInternetAddr> FSocketSubsystemAndroid::GetLocalHostAddr(FOutputDevic
 		int Result = ioctl(TempSocket, SIOCGIFCONF, &IfConfig);
 		if (Result == 0)
 		{
-			int32 WifiAddress = 0;
-			int32 CellularAddress = 0;
-			int32 OtherAddress = 0;
+			sockaddr_storage WifiAddress;
+			sockaddr_storage CellularAddress;
+			sockaddr_storage OtherAddress;
+
+			// Clear these out
+			FMemory::Memzero(&WifiAddress, sizeof(sockaddr_storage));
+			FMemory::Memzero(&CellularAddress, sizeof(sockaddr_storage));
+			FMemory::Memzero(&OtherAddress, sizeof(sockaddr_storage));
 
 			for (int32 IdxReq = 0; IdxReq < ARRAY_COUNT(IfReqs); ++IdxReq)
 			{
@@ -166,45 +163,46 @@ TSharedRef<FInternetAddr> FSocketSubsystemAndroid::GetLocalHostAddr(FOutputDevic
 					if (strcmp(IfReqs[IdxReq].ifr_name, "wlan0") == 0)
 					{
 						// 'Usually' wifi, Prefer wifi
-						WifiAddress = reinterpret_cast<sockaddr_in *>(&IfReqs[IdxReq].ifr_addr)->sin_addr.s_addr;
+						FMemory::Memcpy((void*)&WifiAddress, (void*)IfReqs[IdxReq].ifr_addr, sizeof(sockaddr_in));
 						break;
 					}
 					else if (strcmp(IfReqs[IdxReq].ifr_name, "rmnet0") == 0)
 					{
 						// 'Usually' cellular
-						CellularAddress = reinterpret_cast<sockaddr_in *>(&IfReqs[IdxReq].ifr_addr)->sin_addr.s_addr;
+						FMemory::Memcpy((void*)&CellularAddress, (void*)IfReqs[IdxReq].ifr_addr, sizeof(sockaddr_in));
 					}
-					else if (OtherAddress == 0)
+					else if (OtherAddress.ss_family == AF_UNSPEC)
 					{
 						// First alternate found
-						OtherAddress = reinterpret_cast<sockaddr_in *>(&IfReqs[IdxReq].ifr_addr)->sin_addr.s_addr;
+						FMemory::Memcpy((void*)&OtherAddress, (void*)IfReqs[IdxReq].ifr_addr, sizeof(sockaddr_in));
 					}
 				}
 			}
 
+			FInternetAddrBSD& NewAddrRef = static_cast<FInternetAddrBSD&>(Addr.Get());
 			// Prioritize results found
-			if (WifiAddress != 0)
+			if (WifiAddress.ss_family != AF_UNSPEC)
 			{
 				// Prefer Wifi
-				Addr->SetIp(ntohl(WifiAddress));
+				NewAddrRef.SetIp(WifiAddress);
 				UE_LOG(LogSockets, Log, TEXT("(%s) Wifi Adapter IP %s"), GetSocketAPIName(), *Addr->ToString(false));
 			}
-			else if (CellularAddress != 0)
+			else if (CellularAddress.ss_family != AF_UNSPEC)
 			{
 				// Then cellular
-				Addr->SetIp(ntohl(CellularAddress));
+				NewAddrRef.SetIp(CellularAddress);
 				UE_LOG(LogSockets, Log, TEXT("(%s) Cellular Adapter IP %s"), GetSocketAPIName(), *Addr->ToString(false));
 			}
-			else if (OtherAddress != 0)
+			else if (OtherAddress.ss_family != AF_UNSPEC)
 			{
 				// Then whatever else was found
-				Addr->SetIp(ntohl(OtherAddress));
+				NewAddrRef.SetIp(OtherAddress);
 				UE_LOG(LogSockets, Log, TEXT("(%s) Adapter IP %s"), GetSocketAPIName(), *Addr->ToString(false));
 			}
 			else
 			{
 				// Give up
-				Addr->SetIp(0x7f000001);  // 127.0.0.1
+				Addr->SetLoopbackAddress();  // 127.0.0.1
 				UE_LOG(LogSockets, Warning, TEXT("(%s) NO 'UP' ADAPTER FOUND! using: %s"), GetSocketAPIName(), *Addr->ToString(false));
 			}
 		}
diff --git a/Engine/Source/Runtime/Sockets/Private/BSDIPv6Sockets/IPAddressBSDIPv6.cpp b/Engine/Source/Runtime/Sockets/Private/BSDIPv6Sockets/IPAddressBSDIPv6.cpp
new file mode 100644
--- /dev/null
+++ b/Engine/Source/Runtime/Sockets/Private/BSDIPv6Sockets/IPAddressBSDIPv6.cpp
@@ -0,0 +1,345 @@
+// Copyright 1998-2018 Epic Games, Inc. All Rights Reserved.
+
+#include "IPAddressBSDIPv6.h"
+
+#if PLATFORM_HAS_BSD_IPV6_SOCKETS
+
+PRAGMA_DISABLE_DEPRECATION_WARNINGS
+
+FInternetAddrBSDIPv6::FInternetAddrBSDIPv6()
+{
+	FMemory::Memzero(&Addr, sizeof(Addr));
+	Addr.sin6_family = AF_INET6;
+}
+
+void FInternetAddrBSDIPv6::SetRawIp(const TArray<uint8>& RawAddr)
+{
+	if (RawAddr.Num() >= 16)
+	{
+		for (int i = 0; i < 16; ++i)
+		{
+			Addr.sin6_addr.s6_addr[i] = RawAddr[i];
+		}
+
+		// If this address has an interface, we'll copy it over too.
+		if (RawAddr.Num() == 20)
+		{
+			Addr.sin6_scope_id = (RawAddr[16] << 24) | (RawAddr[17] << 16) | (RawAddr[18] << 8) | (RawAddr[19] << 0);
+		}
+
+		Addr.sin6_family = AF_INET6;
+	}
+	else
+	{
+		FMemory::Memzero(&Addr, sizeof(Addr));
+		Addr.sin6_family = AF_INET6;
+	}
+}
+
+TArray<uint8> FInternetAddrBSDIPv6::GetRawIp() const
+{
+	TArray<uint8> RawAddressArray;
+	for (int i = 0; i < 16; ++i)
+	{
+		RawAddressArray.Add(Addr.sin6_addr.s6_addr[i]);
+	}
+
+	// Copy over the interface if we have it explicitly set (anything other than 0)
+	if (Addr.sin6_scope_id != 0)
+	{
+		uint32 RawScopeId = Addr.sin6_scope_id;
+		RawAddressArray.Add((RawScopeId >> 24) & 0xFF);
+		RawAddressArray.Add((RawScopeId >> 16) & 0xFF);
+		RawAddressArray.Add((RawScopeId >> 8) & 0xFF);
+		RawAddressArray.Add((RawScopeId >> 0) & 0xFF);
+	}
+	return RawAddressArray;
+}
+
+void FInternetAddrBSDIPv6::SetIp(uint32 InAddr)
+{
+	if (InAddr == 0)
+	{
+		FMemory::Memzero(&Addr.sin6_addr, sizeof(Addr.sin6_addr));
+	}
+	else if (InAddr == IPv4MulitcastAddr)
+	{
+		// hack: if it's the hardcoded IPv4 multicasting address then translate into an IPv6 multicast address
+		bool isValid;
+		SetIp(TEXT("ff02::2"), isValid);
+		check(isValid);
+	}
+	else
+	{
+		in_addr InternetAddr;
+		InternetAddr.s_addr = htonl(InAddr);
+
+		SetIp(InternetAddr);
+	}
+}
+
+void FInternetAddrBSDIPv6::SetIp(const TCHAR* InAddr, bool& bIsValid)
+{
+	bIsValid = false;
+
+	FString AddressString(InAddr);
+
+	const bool bHasOpenBracket = AddressString.Contains("[");
+	const int32 CloseBracketIndex = AddressString.Find("]");
+	const bool bHasCloseBracket = CloseBracketIndex != -1;
+
+	// IPv6 may or may not include open and close brackets.
+	// However, only IPv6 address can have them.
+	bool bIsIPv6 = bHasOpenBracket && bHasCloseBracket;
+
+	// Valid IPv4 should not contain an open or close bracket
+	const bool bIsLikelyIPv4 = !bHasOpenBracket && !bHasCloseBracket;
+
+	if (bIsLikelyIPv4 || bIsIPv6)
+	{
+		const int32 LastColonIndex = AddressString.Find(":", ESearchCase::IgnoreCase, ESearchDir::FromEnd);
+
+		// Double check to ensure this isn't actually an IPv6 address without brackets.
+		if (bIsLikelyIPv4)
+		{
+			// IPv4 addresses can contain at most 1 colon.
+			const int32 FirstColonIndex = AddressString.Find(":");
+			bIsIPv6 = (FirstColonIndex != LastColonIndex);
+		}
+
+		// IPv4 address will only have a port when a colon is present.
+		// IPv6 address will only have a port when surrounded by brackets.
+		const bool bHasPort = (INDEX_NONE != LastColonIndex) && (!bIsIPv6 || (bHasCloseBracket && LastColonIndex > CloseBracketIndex));
+		FString Port;
+		if (bHasPort)
+		{
+			Port = AddressString.RightChop(LastColonIndex + 1);
+			AddressString = AddressString.Left(LastColonIndex);
+		}
+
+		if (bIsIPv6)
+		{
+			AddressString.RemoveFromStart("[");
+			AddressString.RemoveFromEnd("]");
+
+			// Check for valid IPv6 address
+			const auto InAddrAnsi = StringCast<ANSICHAR>(*AddressString);
+#if PLATFORM_IOS
+			ISocketSubsystem* SocketSubsystem = ISocketSubsystem::Get(PLATFORM_SOCKETSUBSYSTEM);
+			if (SocketSubsystem && SocketSubsystem->GetHostByName(InAddrAnsi.Get(), *this) == SE_NO_ERROR)
+			{
+				bIsValid = true;
+			}
+			else
+#endif
+				if (inet_pton(AF_INET6, InAddrAnsi.Get(), &Addr.sin6_addr))
+				{
+					bIsValid = true;
+				}
+		}
+		else
+		{
+			const auto InAddrAnsi = StringCast<ANSICHAR>(*AddressString);
+			in_addr IPv4Addr;
+#if PLATFORM_IOS
+			ISocketSubsystem* SocketSubsystem = ISocketSubsystem::Get(PLATFORM_SOCKETSUBSYSTEM);
+			if (SocketSubsystem && SocketSubsystem->GetHostByName(InAddrAnsi.Get(), *this) == SE_NO_ERROR)
+			{
+				bIsValid = true;
+			}
+			else
+#endif
+				if (inet_pton(AF_INET, InAddrAnsi.Get(), &IPv4Addr))
+				{
+					bIsValid = true;
+					SetIp(IPv4Addr);
+				}
+		}
+
+		if (bHasPort)
+		{
+			SetPort(FCString::Atoi(*Port));
+		}
+	}
+}
+
+void FInternetAddrBSDIPv6::SetIp(const in_addr& IPv4Addr)
+{
+	FMemory::Memzero(&Addr.sin6_addr, sizeof(Addr.sin6_addr));
+
+	// map IPv4 to IPv6 address (only works on hybrid network stacks)
+	Addr.sin6_addr.s6_addr[10] = 0xff;
+	Addr.sin6_addr.s6_addr[11] = 0xff;
+
+	uint8	IPv4b1 = (static_cast<uint32>(IPv4Addr.s_addr) & 0xFF),
+		IPv4b2 = ((static_cast<uint32>(IPv4Addr.s_addr) >> 8) & 0xFF),
+		IPv4b3 = ((static_cast<uint32>(IPv4Addr.s_addr) >> 16) & 0xFF),
+		IPv4b4 = ((static_cast<uint32>(IPv4Addr.s_addr) >> 24) & 0xFF);
+
+	Addr.sin6_addr.s6_addr[12] = IPv4b1;
+	Addr.sin6_addr.s6_addr[13] = IPv4b2;
+	Addr.sin6_addr.s6_addr[14] = IPv4b3;
+	Addr.sin6_addr.s6_addr[15] = IPv4b4;
+
+	UE_LOG(LogSockets, Verbose, TEXT("Using IPv4 address: %d.%d.%d.%d  on an ipv6 socket"),
+		IPv4b1,
+		IPv4b2,
+		IPv4b3,
+		IPv4b4
+	);
+}
+
+void FInternetAddrBSDIPv6::SetIp(const in6_addr& IpAddr)
+{
+	Addr.sin6_addr = IpAddr;
+}
+
+void FInternetAddrBSDIPv6::SetIp(const sockaddr_storage& IpAddr)
+{
+	if (IpAddr.ss_family == AF_INET)
+	{
+		const sockaddr_in* SockAddr = (const sockaddr_in*)&IpAddr;
+		SetIp(SockAddr->sin_addr);
+	}
+	else if (IpAddr.ss_family == AF_INET6)
+	{
+		const sockaddr_in6* SockAddr = (const sockaddr_in6*)&IpAddr;
+		SetIp(SockAddr->sin6_addr);
+	}
+}
+
+void FInternetAddrBSDIPv6::GetIp(uint32& OutAddr) const
+{
+	// grab the last 32-bits of the IPv6 address as this will correspond to the IPv4 address
+	// in a dual stack system.
+	// This function doesn't really make sense in IPv6, but too much other code relies on it
+	// existing to not have this here.
+#if PLATFORM_LITTLE_ENDIAN
+	OutAddr = (Addr.sin6_addr.s6_addr[12] << 24) | (Addr.sin6_addr.s6_addr[13] << 16) | (Addr.sin6_addr.s6_addr[14] << 8) | (Addr.sin6_addr.s6_addr[15]);
+#else
+	OutAddr = (Addr.sin6_addr.s6_addr[15] << 24) | (Addr.sin6_addr.s6_addr[14] << 16) | (Addr.sin6_addr.s6_addr[13] << 8) | (Addr.sin6_addr.s6_addr[12]);
+#endif
+}
+
+void FInternetAddrBSDIPv6::GetIp(in6_addr& OutAddr) const
+{
+	OutAddr = Addr.sin6_addr;
+}
+
+void FInternetAddrBSDIPv6::SetPort(int32 InPort)
+{
+	Addr.sin6_port = htons(InPort);
+}
+
+int32 FInternetAddrBSDIPv6::GetPort() const
+{
+	return ntohs(Addr.sin6_port);
+}
+
+void FInternetAddrBSDIPv6::SetAnyAddress()
+{
+	SetIp(in6addr_any);
+	SetPort(0);
+}
+
+void FInternetAddrBSDIPv6::SetBroadcastAddress()
+{
+	// broadcast means something different in IPv6, but this is a rough equivalent
+#ifndef in6addr_allnodesonlink
+	// see RFC 4291, link-local multicast address http://tools.ietf.org/html/rfc4291
+	static in6_addr in6addr_allnodesonlink =
+	{
+		{ { 0xff, 0x02, 0x00, 0x00, 0x00, 0x00, 0x00, 0x00, 0x00, 0x00, 0x00, 0x00, 0x00, 0x00, 0x00, 0x01 } }
+	};
+#endif // in6addr_allnodesonlink
+
+	SetIp(in6addr_allnodesonlink);
+	SetPort(0);
+}
+
+void FInternetAddrBSDIPv6::SetLoopbackAddress()
+{
+	SetIp(in6addr_loopback);
+	SetPort(0);
+}
+
+FString FInternetAddrBSDIPv6::ToString(bool bAppendPort) const
+{
+	char IPStr[INET6_ADDRSTRLEN];
+
+	inet_ntop(AF_INET6, (void*)&Addr.sin6_addr, IPStr, INET6_ADDRSTRLEN);
+
+	FString Result("[");
+	Result += IPStr;
+	Result += "]";
+
+	if (bAppendPort)
+	{
+		Result += ":";
+		Result += FString::Printf(TEXT("%d"), GetPort());
+	}
+
+	return Result;
+}
+
+uint32 FInternetAddrBSDIPv6::GetTypeHash()
+{
+	// @todo: Find a more efficient way to hash IPv6 addresses, if they are to be used with NetConnection's
+	//return Addr.sin_addr.s_addr + (Addr.sin_port * 23);
+
+	return ::GetTypeHash(*ToString(true));
+}
+
+bool FInternetAddrBSDIPv6::IsValid() const
+{
+	FInternetAddrBSDIPv6 Temp;
+	return memcmp(&Addr.sin6_addr, &Temp.Addr.sin6_addr, sizeof(in6_addr)) != 0;
+}
+
+TSharedRef<FInternetAddr> FInternetAddrBSDIPv6::Clone() const
+{
+
+	FSocketSubsystemBSDIPv6* SocketSubsystem = static_cast<FSocketSubsystemBSDIPv6*>(ISocketSubsystem::Get());
+	check(SocketSubsystem);
+
+	TSharedRef<FInternetAddr> NewAddress = SocketSubsystem->CreateInternetAddr();
+	((FInternetAddrBSDIPv6*)(&(NewAddress.Get())))->Addr = Addr;
+	
+	return NewAddress;
+}
+
+bool FInternetAddrBSDIPv6::operator==(const FInternetAddr& Other) const
+{
+	FInternetAddrBSDIPv6& OtherBSD = (FInternetAddrBSDIPv6&)Other;
+	return memcmp(&Addr.sin6_addr, &OtherBSD.Addr.sin6_addr, sizeof(in6_addr)) == 0 &&
+		Addr.sin6_port == OtherBSD.Addr.sin6_port &&
+		Addr.sin6_family == OtherBSD.Addr.sin6_family;
+}
+
+FResolveInfoCachedBSDIPv6::FResolveInfoCachedBSDIPv6(const FInternetAddr& InAddr)
+{
+	
+	const FInternetAddrBSDIPv6* InAddrAsIPv6 = static_cast<const FInternetAddrBSDIPv6*>(&InAddr);
+	if (InAddrAsIPv6)
+	{
+		Addr = ISocketSubsystem::Get()->CreateInternetAddr();
+		FInternetAddrBSDIPv6* AddrAsIPv6 = static_cast<FInternetAddrBSDIPv6*>(Addr.Get());
+		if (AddrAsIPv6)
+		{
+			AddrAsIPv6->SetPort(InAddr.GetPort());
+			in6_addr temp;
+			InAddrAsIPv6->GetIp(temp);
+			AddrAsIPv6->SetIp(temp);
+		}
+	}
+	else
+	{
+		uint32 IpAddr;
+		InAddr.GetIp(IpAddr);
+		Addr = ISocketSubsystem::Get()->CreateInternetAddr(IpAddr, InAddr.GetPort());
+	}
+	
+}
+
+PRAGMA_ENABLE_DEPRECATION_WARNINGS
+#endif
\ No newline at end of file
diff --git a/Engine/Source/Runtime/Sockets/Private/BSDIPv6Sockets/IPAddressBSDIPv6.h b/Engine/Source/Runtime/Sockets/Private/BSDIPv6Sockets/IPAddressBSDIPv6.h
--- a/Engine/Source/Runtime/Sockets/Private/BSDIPv6Sockets/IPAddressBSDIPv6.h
+++ b/Engine/Source/Runtime/Sockets/Private/BSDIPv6Sockets/IPAddressBSDIPv6.h
@@ -4,59 +4,53 @@
 
 #include "CoreMinimal.h"
 #include "BSDSockets/SocketSubsystemBSDPrivate.h"
+#include "SocketSubsystemBSDIPv6.h"
 #include "IPAddress.h"
 
 #if PLATFORM_HAS_BSD_IPV6_SOCKETS
 
-
-
-/**
-* Represents an internet ip address, using the relatively standard sockaddr_in6 structure. All data is in network byte order
-*/
-class FInternetAddrBSDIPv6 : public FInternetAddr
+class FInternetAddrBSDIPv6_Base_DEPRECATED
 {
+protected:
 	/** The internet ip address structure */
 	sockaddr_in6 Addr;
 
 	/** Horrible hack to catch hard coded multicasting on IPv4 **/
 	static const uint32 IPv4MulitcastAddr = ((230 << 24) | (0 << 16) | (0 << 8) | (1 << 0));
+};
+
+/**
+* Represents an internet ip address, using the relatively standard sockaddr_in6 structure. All data is in network byte order
+*/
+class DEPRECATED(4.21, "FInternetAddrBSDIPv6 is deprecated. Move to FInternetAddrBSD.") FInternetAddrBSDIPv6 : public FInternetAddr, FInternetAddrBSDIPv6_Base_DEPRECATED
+{
 
 public:
 	/**
-	* Constructor. Sets address to default state
-	*/
-	FInternetAddrBSDIPv6()
-	{
-		FMemory::Memzero(&Addr, sizeof(Addr));
-		Addr.sin6_family = AF_INET6;
-	}
+	 * Constructor. Sets address to default state
+	 */
+	FInternetAddrBSDIPv6();
 
 	/**
-	* Sets the ip address from a host byte order uint32, convert the IPv4 address supplied to an IPv6 address
-	*
-	* @param InAddr the new address to use (must convert to network byte order)
-	*/
-	virtual void SetIp(uint32 InAddr) override
-	{
-		if (InAddr == 0)
-		{
-			FMemory::Memzero(&Addr.sin6_addr, sizeof(Addr.sin6_addr));
-		}
-		else if (InAddr == IPv4MulitcastAddr)
-		{
-			// hack: if it's the hardcoded IPv4 multicasting address then translate into an IPv6 multicast address
-			bool isValid;
-			SetIp(TEXT("ff02::2"), isValid);
-			check(isValid);
-		}
-		else
-		{
-			in_addr InternetAddr;
-			InternetAddr.s_addr = htonl(InAddr);
-
-			SetIp(InternetAddr);
-		}
-	}
+	 * Sets the ip address from a raw network byte order array.
+	 *
+	 * @param RawAddr the new address to use (must be converted to network byte order)
+	 */
+	virtual void SetRawIp(const TArray<uint8>& RawAddr) override;
+
+	/**
+	 * Gets the ip address in a raw array stored in network byte order.
+	 *
+	 * @return The raw address stored in an array
+	 */
+	virtual TArray<uint8> GetRawIp() const override;
+
+	/**
+	 * Sets the ip address from a host byte order uint32, convert the IPv4 address supplied to an IPv6 address
+	 *
+	 * @param InAddr the new address to use (must convert to network byte order)
+	 */
+	virtual void SetIp(uint32 InAddr) override;
 
 	/**
 	 * Sets the ip address from a string IPv6 or IPv4 address.
@@ -68,131 +62,21 @@ public:
 	 * @param InAddr the string containing the new ip address to use
 	 * @param bIsValid will be set to true if InAddr was a valid IPv6 or IPv4 address, false if not.
 	 */
-	virtual void SetIp(const TCHAR* InAddr, bool& bIsValid) override
-	{
-		bIsValid = false;
-
-		FString AddressString(InAddr);
-
-		const bool bHasOpenBracket = AddressString.Contains("[");
-		const int32 CloseBracketIndex = AddressString.Find("]");
-		const bool bHasCloseBracket = CloseBracketIndex != -1;
-
-		// IPv6 may or may not include open and close brackets.
-		// However, only IPv6 address can have them.
-		bool bIsIPv6 = bHasOpenBracket && bHasCloseBracket;
-
-		// Valid IPv4 should not contain an open or close bracket
-		const bool bIsLikelyIPv4 = !bHasOpenBracket && !bHasCloseBracket;
-
-		if (bIsLikelyIPv4 || bIsIPv6)
-		{
-			const int32 LastColonIndex = AddressString.Find(":", ESearchCase::IgnoreCase, ESearchDir::FromEnd);
-
-			// Double check to ensure this isn't actually an IPv6 address without brackets.
-			if (bIsLikelyIPv4)
-			{
-				// IPv4 addresses can contain at most 1 colon.
-				const int32 FirstColonIndex = AddressString.Find(":");
-				bIsIPv6 = (FirstColonIndex != LastColonIndex);
-			}
-
-			// IPv4 address will only have a port when a colon is present.
-			// IPv6 address will only have a port when surrounded by brackets.
-			const bool bHasPort = (INDEX_NONE != LastColonIndex) && (!bIsIPv6 || (bHasCloseBracket && LastColonIndex > CloseBracketIndex));
-			FString Port;
-			if (bHasPort)
-			{
-				Port = AddressString.RightChop(LastColonIndex + 1);
-				AddressString = AddressString.Left(LastColonIndex);
-			}
-
-			if (bIsIPv6)
-			{
-				AddressString.RemoveFromStart("[");
-				AddressString.RemoveFromEnd("]");
-
-				// Check for valid IPv6 address
-				const auto InAddrAnsi = StringCast<ANSICHAR>(*AddressString);
-#if PLATFORM_IOS
-				ISocketSubsystem* SocketSubsystem = ISocketSubsystem::Get(PLATFORM_SOCKETSUBSYSTEM);
-				if (SocketSubsystem && SocketSubsystem->GetHostByName(InAddrAnsi.Get(), *this) == SE_NO_ERROR)
-				{
-					bIsValid = true;
-				}
-				else
-#endif
-				if (inet_pton(AF_INET6, InAddrAnsi.Get(), &Addr.sin6_addr))
-				{
-					bIsValid = true;
-				}
-			}
-			else
-			{
-				const auto InAddrAnsi = StringCast<ANSICHAR>(*AddressString);
-				in_addr IPv4Addr;
-#if PLATFORM_IOS
-				ISocketSubsystem* SocketSubsystem = ISocketSubsystem::Get(PLATFORM_SOCKETSUBSYSTEM);
-				if (SocketSubsystem && SocketSubsystem->GetHostByName(InAddrAnsi.Get(), *this) == SE_NO_ERROR)
-				{
-					bIsValid = true;
-				}
-				else
-#endif
-				if (inet_pton(AF_INET, InAddrAnsi.Get(), &IPv4Addr))
-				{
-					bIsValid = true;
-					SetIp(IPv4Addr);
-				}
-			}
-
-			if (bHasPort)
-			{
-				SetPort(FCString::Atoi(*Port));
-			}
-		}
-	}
+	virtual void SetIp(const TCHAR* InAddr, bool& bIsValid) override;
 
 	/**
-	* Sets the ip address using a network byte order ipv4 address
-	*
-	* @param IPv4Addr the new ip address to use
-	*/
-	void SetIp(const in_addr& IPv4Addr)
-	{
-		FMemory::Memzero(&Addr.sin6_addr, sizeof(Addr.sin6_addr));
-
-		// map IPv4 to IPv6 address (only works on hybrid network stacks)
-		Addr.sin6_addr.s6_addr[10] = 0xff;
-		Addr.sin6_addr.s6_addr[11] = 0xff;
-
-		uint8	IPv4b1 = (static_cast<uint32>(IPv4Addr.s_addr) & 0xFF),
-			IPv4b2 = ((static_cast<uint32>(IPv4Addr.s_addr) >> 8) & 0xFF),
-			IPv4b3 = ((static_cast<uint32>(IPv4Addr.s_addr) >> 16) & 0xFF),
-			IPv4b4 = ((static_cast<uint32>(IPv4Addr.s_addr) >> 24) & 0xFF);
-
-		Addr.sin6_addr.s6_addr[12] = IPv4b1;
-		Addr.sin6_addr.s6_addr[13] = IPv4b2;
-		Addr.sin6_addr.s6_addr[14] = IPv4b3;
-		Addr.sin6_addr.s6_addr[15] = IPv4b4;
-
-		UE_LOG(LogSockets, Verbose, TEXT("Using IPv4 address: %d.%d.%d.%d  on an ipv6 socket"),
-			IPv4b1,
-			IPv4b2,
-			IPv4b3,
-			IPv4b4
-			);
-	}
+	 * Sets the ip address using a network byte order ipv4 address
+	 *
+	 * @param IPv4Addr the new ip address to use
+	 */
+	void SetIp(const in_addr& IPv4Addr);
 
 	/**
-	* Sets the ip address using a network byte order ipv6 address
-	*
-	* @param IpAddr the new ip address to use
-	*/
-	void SetIp(const in6_addr& IpAddr)
-	{
-		Addr.sin6_addr = IpAddr;
-	}
+	 * Sets the ip address using a network byte order ipv6 address
+	 *
+	 * @param IpAddr the new ip address to use
+	 */
+	void SetIp(const in6_addr& IpAddr);
 	
 #if PLATFORM_IOS
 	void SetIp(const sockaddr_in6& IpAddr)
@@ -202,157 +86,67 @@ public:
 #endif
 
 	/**
-	* Sets the ip address using a generic sockaddr_storage
-	*
-	* @param IpAddr the new ip address to use
-	*/
-	void SetIp(const sockaddr_storage& IpAddr)
-	{
-		if (IpAddr.ss_family == AF_INET)
-		{
-			const sockaddr_in* SockAddr = (const sockaddr_in*)&IpAddr;
-			SetIp(SockAddr->sin_addr);
-		}
-		else if (IpAddr.ss_family == AF_INET6)
-		{
-			const sockaddr_in6* SockAddr = (const sockaddr_in6*)&IpAddr;
-			SetIp(SockAddr->sin6_addr);
-		}
-	}
-
-	/**
-	* Copies the network byte order ip address to a host byte order dword, doesn't exist with IPv6
-	*
-	* @param OutAddr the out param receiving the ip address
-	*/
-	virtual void GetIp(uint32& OutAddr) const override
-	{
-		// grab the last 32-bits of the IPv6 address as this will correspond to the IPv4 address
-		// in a dual stack system.
-		// This function doesn't really make sense in IPv6, but too much other code relies on it
-		// existing to not have this here.
-#if PLATFORM_LITTLE_ENDIAN
-		OutAddr = (Addr.sin6_addr.s6_addr[12] << 24) | (Addr.sin6_addr.s6_addr[13] << 16) | (Addr.sin6_addr.s6_addr[14] << 8) | (Addr.sin6_addr.s6_addr[15]);
-#else
-		OutAddr = (Addr.sin6_addr.s6_addr[15] << 24) | (Addr.sin6_addr.s6_addr[14] << 16) | (Addr.sin6_addr.s6_addr[13] << 8) | (Addr.sin6_addr.s6_addr[12]);
-#endif
-	}
+	 * Sets the ip address using a generic sockaddr_storage
+	 *
+	 * @param IpAddr the new ip address to use
+	 */
+	void SetIp(const sockaddr_storage& IpAddr);
 
 	/**
-	* Copies the network byte order ip address
-	*
-	* @param OutAddr the out param receiving the ip address
-	*/
-	void GetIp(in6_addr& OutAddr) const
-	{
-		OutAddr = Addr.sin6_addr;
-	}
+	 * Copies the network byte order ip address to a host byte order dword, doesn't exist with IPv6
+	 *
+	 * @param OutAddr the out param receiving the ip address
+	 */
+	virtual void GetIp(uint32& OutAddr) const override;
 
 	/**
-	* Sets the port number from a host byte order int
-	*
-	* @param InPort the new port to use (must convert to network byte order)
-	*/
-	virtual void SetPort(int32 InPort) override
-	{
-		Addr.sin6_port = htons(InPort);
-	}
+	 * Copies the network byte order ip address
+	 *
+	 * @param OutAddr the out param receiving the ip address
+	 */
+	void GetIp(in6_addr& OutAddr) const;
 
 	/**
-	* Copies the port number from this address and places it into a host byte order int
-	*
-	* @param OutPort the host byte order int that receives the port
-	*/
-	virtual void GetPort(int32& OutPort) const override
-	{
-		OutPort = ntohs(Addr.sin6_port);
-	}
+	 * Sets the port number from a host byte order int
+	 *
+	 * @param InPort the new port to use (must convert to network byte order)
+	 */
+	virtual void SetPort(int32 InPort) override;
 
 	/** Returns the port number from this address in host byte order */
-	virtual int32 GetPort() const override
-	{
-		return ntohs(Addr.sin6_port);
-	}
+	virtual int32 GetPort() const override;
 
 	/** Sets the address to be any address */
-	virtual void SetAnyAddress() override
-	{
-		SetIp(in6addr_any);
-		SetPort(0);
-	}
+	virtual void SetAnyAddress() override;
 
 	/** Sets the address to broadcast */
-	virtual void SetBroadcastAddress() override
-	{
-		// broadcast means something different in IPv6, but this is a rough equivalent
-#ifndef in6addr_allnodesonlink
-		// see RFC 4291, link-local multicast address http://tools.ietf.org/html/rfc4291
-		static in6_addr in6addr_allnodesonlink =
-		{
-			{ { 0xff, 0x02, 0x00, 0x00, 0x00, 0x00, 0x00, 0x00, 0x00, 0x00, 0x00, 0x00, 0x00, 0x00, 0x00, 0x01 } }
-		};
-#endif // in6addr_allnodesonlink
-
-		SetIp(in6addr_allnodesonlink);
-		SetPort(0);
-	}
-
-	/**
-	* Converts this internet ip address to string form. String will be enclosed in square braces.
-	*
-	* @param bAppendPort whether to append the port information or not
-	*/
-	virtual FString ToString(bool bAppendPort) const override
-	{
-		char IPStr[INET6_ADDRSTRLEN];
-
-		inet_ntop(AF_INET6, (void*)&Addr.sin6_addr, IPStr, INET6_ADDRSTRLEN);
-
-		FString Result("[");
-		Result += IPStr;
-		Result += "]";
+	virtual void SetBroadcastAddress() override;
 
-		if (bAppendPort)
-		{
-			Result += ":";
-			Result += FString::Printf(TEXT("%d"), GetPort());
-		}
-
-		return Result;
-	}
+	/** Sets the address to the loopback address */
+	virtual void SetLoopbackAddress() override;
 
 	/**
-	* Compares two internet ip addresses for equality
-	*
-	* @param Other the address to compare against
-	*/
-	virtual bool operator==(const FInternetAddr& Other) const override
-	{
-		FInternetAddrBSDIPv6& OtherBSD = (FInternetAddrBSDIPv6&)Other;
-		return memcmp(&Addr.sin6_addr, &OtherBSD.Addr.sin6_addr, sizeof(in6_addr)) == 0 &&
-			Addr.sin6_port == OtherBSD.Addr.sin6_port &&
-			Addr.sin6_family == OtherBSD.Addr.sin6_family;
-	}
+	 * Converts this internet ip address to string form. String will be enclosed in square braces.
+	 *
+	 * @param bAppendPort whether to append the port information or not
+	 */
+	virtual FString ToString(bool bAppendPort) const override;
 
-	virtual uint32 GetTypeHash() override
-	{
-		// @todo: Find a more efficient way to hash IPv6 addresses, if they are to be used with NetConnection's
-		//return Addr.sin_addr.s_addr + (Addr.sin_port * 23);
+	/**
+	 * Compares two internet ip addresses for equality
+	 *
+	 * @param Other the address to compare against
+	 */
+	virtual bool operator==(const FInternetAddr& Other) const override;
 
-		return ::GetTypeHash(*ToString(true));
-	}
+	virtual uint32 GetTypeHash() override;
 
 	/**
-	* Is this a well formed internet address, the only criteria being non-zero
-	*
-	* @return true if a valid IP, false otherwise
-	*/
-	virtual bool IsValid() const override
-	{
-		FInternetAddrBSDIPv6 Temp;
-
-		return memcmp(&Addr.sin6_addr, &Temp.Addr.sin6_addr, sizeof(in6_addr)) != 0;
-	}
+	 * Is this a well formed internet address, the only criteria being non-zero
+	 *
+	 * @return true if a valid IP, false otherwise
+	 */
+	virtual bool IsValid() const override;
 
 	operator sockaddr*(void)
 	{
@@ -363,43 +157,17 @@ public:
 	{
 		return (const sockaddr*)&Addr;
 	}
-	
-#if PLATFORM_IOS
-	virtual void Copy(const FInternetAddr& Other) override
-	{
-		Addr = static_cast<const FInternetAddrBSDIPv6&>(Other).Addr;
-	}
-#endif
+
+	virtual TSharedRef<FInternetAddr> Clone() const override;
 };
 
 
-class FResolveInfoCachedBSDIPv6 : public FResolveInfoCached
+class DEPRECATED(4.21, "This class is no longer needed as the base class can handle the proper construction now.") FResolveInfoCachedBSDIPv6 : public FResolveInfoCached
 {
 	FResolveInfoCachedBSDIPv6();
 
 public:
-	FResolveInfoCachedBSDIPv6(const FInternetAddr& InAddr)
-	{
-		const FInternetAddrBSDIPv6* InAddrAsIPv6 = static_cast<const FInternetAddrBSDIPv6*>(&InAddr);
-		if (InAddrAsIPv6)
-		{
-			Addr = ISocketSubsystem::Get()->CreateInternetAddr();
-			FInternetAddrBSDIPv6* AddrAsIPv6 = static_cast<FInternetAddrBSDIPv6*>(Addr.Get());
-			if (AddrAsIPv6)
-			{
-				AddrAsIPv6->SetPort(InAddr.GetPort());
-				in6_addr temp;
-				InAddrAsIPv6->GetIp(temp);
-				AddrAsIPv6->SetIp(temp);
-			}
-		}
-		else
-		{
-			uint32 IpAddr;
-			InAddr.GetIp(IpAddr);
-			Addr = ISocketSubsystem::Get()->CreateInternetAddr(IpAddr, InAddr.GetPort());
-		}
-	}
+	FResolveInfoCachedBSDIPv6(const FInternetAddr& InAddr);
 };
 
 #endif
diff --git a/Engine/Source/Runtime/Sockets/Private/BSDIPv6Sockets/SocketSubsystemBSDIPv6.cpp b/Engine/Source/Runtime/Sockets/Private/BSDIPv6Sockets/SocketSubsystemBSDIPv6.cpp
--- a/Engine/Source/Runtime/Sockets/Private/BSDIPv6Sockets/SocketSubsystemBSDIPv6.cpp
+++ b/Engine/Source/Runtime/Sockets/Private/BSDIPv6Sockets/SocketSubsystemBSDIPv6.cpp
@@ -7,13 +7,45 @@
 
 #include "BSDIPv6Sockets/IPAddressBSDIPv6.h"
 
-
+PRAGMA_DISABLE_DEPRECATION_WARNINGS
 class FSocketBSDIPv6* FSocketSubsystemBSDIPv6::InternalBSDSocketFactory(SOCKET Socket, ESocketType SocketType, const FString& SocketDescription)
 {
 	// return a new socket object
 	return new FSocketBSDIPv6(Socket, SocketType, SocketDescription, this);
 }
 
+ESocketErrors FSocketSubsystemBSDIPv6::TranslateGAIErrorCode(int32 Code) const
+{
+#if PLATFORM_HAS_BSD_SOCKET_FEATURE_GETADDRINFO
+	switch (Code)
+	{
+		// getaddrinfo() has its own error codes
+	case EAI_AGAIN:			return SE_TRY_AGAIN;
+	case EAI_BADFLAGS:		return SE_EINVAL;
+	case EAI_FAIL:			return SE_NO_RECOVERY;
+	case EAI_FAMILY:		return SE_EAFNOSUPPORT;
+	case EAI_MEMORY:		return SE_ENOBUFS;
+	case EAI_NONAME:		return SE_HOST_NOT_FOUND;
+	case EAI_SERVICE:		return SE_EPFNOSUPPORT;
+	case EAI_SOCKTYPE:		return SE_ESOCKTNOSUPPORT;
+#if PLATFORM_HAS_BSD_SOCKET_FEATURE_WINSOCKETS
+	case WSANO_DATA:		return SE_NO_DATA;
+	case WSANOTINITIALISED: return SE_NOTINITIALISED;
+#else			
+	case EAI_NODATA:		return SE_NO_DATA;
+	case EAI_ADDRFAMILY:	return SE_ADDRFAMILY;
+	case EAI_SYSTEM:		return SE_SYSTEM;
+#endif
+	case 0:					break; // 0 means success
+	default:
+		UE_LOG(LogSockets, Warning, TEXT("Unhandled getaddrinfo() socket error! Code: %d"), Code);
+		return SE_EINVAL;
+	}
+#endif // PLATFORM_HAS_BSD_SOCKET_FEATURE_GETADDRINFO
+
+	return SE_NO_ERROR;
+}
+
 FSocket* FSocketSubsystemBSDIPv6::CreateSocket(const FName& SocketType, const FString& SocketDescription, bool bForceUDP)
 {
 	SOCKET Socket = INVALID_SOCKET;
@@ -200,5 +232,5 @@ ESocketErrors FSocketSubsystemBSDIPv6::TranslateErrorCode(int32 Code)
 	ensure(0);
 	return SE_EINVAL;
 }
-
+PRAGMA_ENABLE_DEPRECATION_WARNINGS
 #endif
diff --git a/Engine/Source/Runtime/Sockets/Private/BSDIPv6Sockets/SocketSubsystemBSDIPv6.h b/Engine/Source/Runtime/Sockets/Private/BSDIPv6Sockets/SocketSubsystemBSDIPv6.h
--- a/Engine/Source/Runtime/Sockets/Private/BSDIPv6Sockets/SocketSubsystemBSDIPv6.h
+++ b/Engine/Source/Runtime/Sockets/Private/BSDIPv6Sockets/SocketSubsystemBSDIPv6.h
@@ -11,13 +11,11 @@
 
 #if PLATFORM_HAS_BSD_IPV6_SOCKETS
 
-
-
 /**
  * Standard BSD specific IPv6 socket subsystem implementation
  */
-class FSocketSubsystemBSDIPv6
-	: public FSocketSubsystemBSDCommon
+class DEPRECATED(4.21, "Move to FSocketSubsystemBSD") FSocketSubsystemBSDIPv6
+	: public ISocketSubsystem
 {
 public:
 
@@ -70,10 +68,14 @@ protected:
 	 */
 	virtual class FSocketBSDIPv6* InternalBSDSocketFactory( SOCKET Socket, ESocketType SocketType, const FString& SocketDescription );
 
+	PRAGMA_DISABLE_DEPRECATION_WARNINGS
 	// allow BSD sockets to use this when creating new sockets from accept() etc
 	friend FSocketBSDIPv6;
+	PRAGMA_ENABLE_DEPRECATION_WARNINGS
 
 	// Used to prevent multiple threads accessing the shared data.
 	FCriticalSection HostByNameSynch;
+
+	ESocketErrors TranslateGAIErrorCode(int32 Code) const;
 };
 #endif
diff --git a/Engine/Source/Runtime/Sockets/Private/BSDIPv6Sockets/SocketsBSDIPv6.cpp b/Engine/Source/Runtime/Sockets/Private/BSDIPv6Sockets/SocketsBSDIPv6.cpp
--- a/Engine/Source/Runtime/Sockets/Private/BSDIPv6Sockets/SocketsBSDIPv6.cpp
+++ b/Engine/Source/Runtime/Sockets/Private/BSDIPv6Sockets/SocketsBSDIPv6.cpp
@@ -6,6 +6,19 @@
 
 #include "BSDIPv6Sockets/IPAddressBSDIPv6.h"
 #include "BSDIPv6Sockets/SocketSubsystemBSDIPv6.h"
+
+PRAGMA_DISABLE_DEPRECATION_WARNINGS
+
+FSocketBSDIPv6::~FSocketBSDIPv6()
+{
+	FSocketBSDIPv6::Close();
+}
+
+SOCKET FSocketBSDIPv6::GetNativeSocket()
+{
+	return Socket;
+}
+
 //#include "Net/NetworkProfiler.h"
 
 bool FSocketBSDIPv6::Close(void)
@@ -120,7 +133,7 @@ bool FSocketBSDIPv6::HasPendingData(uint32& PendingDataSize)
 }
 
 
-FSocket* FSocketBSDIPv6::Accept(const FString& SocketDescription)
+FSocket* FSocketBSDIPv6::Accept(const FString& InSocketDescription)
 {
 	SOCKET NewSocket = accept(Socket,NULL,NULL);
 
@@ -129,14 +142,14 @@ FSocket* FSocketBSDIPv6::Accept(const FString& SocketDescription)
 		// we need the subclass to create the actual FSocket object
 		check(SocketSubsystem);
 		FSocketSubsystemBSDIPv6* BSDSystem = static_cast<FSocketSubsystemBSDIPv6*>(SocketSubsystem);
-		return BSDSystem->InternalBSDSocketFactory(NewSocket, SocketType, SocketDescription);
+		return BSDSystem->InternalBSDSocketFactory(NewSocket, SocketType, InSocketDescription);
 	}
 
 	return NULL;
 }
 
 
-FSocket* FSocketBSDIPv6::Accept(FInternetAddr& OutAddr, const FString& SocketDescription)
+FSocket* FSocketBSDIPv6::Accept(FInternetAddr& OutAddr, const FString& InSocketDescription)
 {
 	SOCKLEN SizeOf = sizeof(sockaddr_in6);
 	SOCKET NewSocket = accept(Socket, *(FInternetAddrBSDIPv6*)(&OutAddr), &SizeOf);
@@ -146,7 +159,7 @@ FSocket* FSocketBSDIPv6::Accept(FInternetAddr& OutAddr, const FString& SocketDes
 		// we need the subclass to create the actual FSocket object
 		check(SocketSubsystem);
 		FSocketSubsystemBSDIPv6* BSDSystem = static_cast<FSocketSubsystemBSDIPv6*>(SocketSubsystem);
-		return BSDSystem->InternalBSDSocketFactory(NewSocket, SocketType, SocketDescription);
+		return BSDSystem->InternalBSDSocketFactory(NewSocket, SocketType, InSocketDescription);
 	}
 
 	return NULL;
@@ -455,4 +468,6 @@ bool FSocketBSDIPv6::SetIPv6Only(bool bIPv6Only)
 	return bOk;
 }
 
+PRAGMA_ENABLE_DEPRECATION_WARNINGS
+
 #endif
diff --git a/Engine/Source/Runtime/Sockets/Private/BSDIPv6Sockets/SocketsBSDIPv6.h b/Engine/Source/Runtime/Sockets/Private/BSDIPv6Sockets/SocketsBSDIPv6.h
--- a/Engine/Source/Runtime/Sockets/Private/BSDIPv6Sockets/SocketsBSDIPv6.h
+++ b/Engine/Source/Runtime/Sockets/Private/BSDIPv6Sockets/SocketsBSDIPv6.h
@@ -26,11 +26,26 @@ namespace EIPv6SocketInternalState
 	};
 };
 
+class FSocketBSDIPv6_Base_DEPRECATED
+{
+protected:
+
+	FSocketBSDIPv6_Base_DEPRECATED(SOCKET InSocket, ISocketSubsystem* InSubsystem) :
+		Socket(InSocket), SocketSubsystem(InSubsystem) { }
+
+	// Holds the BSD socket object.
+	SOCKET Socket;
+
+	// Pointer to the subsystem that created it
+	ISocketSubsystem * SocketSubsystem;
+};
+
+
 /**
  * Implements a BSD network socket.
  */
-class FSocketBSDIPv6
-	: public FSocket
+class DEPRECATED(4.21, "FSocketBSDIPv6 is deprecated. Move to FSocketBSD.") FSocketBSDIPv6
+	: public FSocket, FSocketBSDIPv6_Base_DEPRECATED
 {
 public:
 
@@ -44,8 +59,7 @@ public:
 	 */
 	FSocketBSDIPv6(SOCKET InSocket, ESocketType InSocketType, const FString& InSocketDescription, ISocketSubsystem* InSubsystem) 
 		: FSocket(InSocketType, InSocketDescription)
-		, Socket(InSocket)
-		, SocketSubsystem(InSubsystem)
+		, FSocketBSDIPv6_Base_DEPRECATED(InSocket, InSubsystem)
 	{ }
 
 	/**
@@ -53,20 +67,14 @@ public:
 	 *
 	 * Closes the socket if it is still open
 	 */
-	virtual ~FSocketBSDIPv6()
-	{
-		FSocketBSDIPv6::Close();
-	}
+	virtual ~FSocketBSDIPv6();
 
 	/**
 	* Gets the Socket for anyone who knows they have an FSocketBSD.
 	*
 	* @return The native socket.
 	*/
-	SOCKET GetNativeSocket()
-	{
-		return Socket;
-	}
+	SOCKET GetNativeSocket();
 
 
 public:
@@ -83,9 +91,9 @@ public:
 
 	virtual bool HasPendingData(uint32& PendingDataSize) override;
 
-	virtual class FSocket* Accept(const FString& SocketDescription) override;
+	virtual class FSocket* Accept(const FString& InSocketDescription) override;
 
-	virtual class FSocket* Accept(FInternetAddr& OutAddr, const FString& SocketDescription) override;
+	virtual class FSocket* Accept(FInternetAddr& OutAddr, const FString& InSocketDescription) override;
 
 	virtual bool SendTo(const uint8* Data, int32 Count, int32& BytesSent, const FInternetAddr& Destination) override;
 
@@ -135,11 +143,6 @@ protected:
 	// This is generally select(), but makes it easier for platforms without select to replace it
 	virtual EIPv6SocketInternalState::Return HasState(EIPv6SocketInternalState::Param State, FTimespan WaitTime = FTimespan::Zero());
 
-	// Holds the BSD socket object.
-	SOCKET Socket;
-
-	// Pointer to the subsystem that created it
-	ISocketSubsystem * SocketSubsystem;
 };
 
 #endif
diff --git a/Engine/Source/Runtime/Sockets/Private/BSDSockets/IPAddressBSD.cpp b/Engine/Source/Runtime/Sockets/Private/BSDSockets/IPAddressBSD.cpp
new file mode 100644
--- /dev/null
+++ b/Engine/Source/Runtime/Sockets/Private/BSDSockets/IPAddressBSD.cpp
@@ -0,0 +1,407 @@
+// Copyright 1998-2018 Epic Games, Inc. All Rights Reserved.
+
+#include "IPAddressBSD.h"
+#include "SocketSubsystemBSD.h"
+
+#if PLATFORM_HAS_BSD_SOCKETS || PLATFORM_HAS_BSD_IPV6_SOCKETS
+
+FInternetAddrBSD::FInternetAddrBSD()
+{
+	Clear();
+}
+
+void FInternetAddrBSD::Clear()
+{
+	FMemory::Memzero(&Addr, sizeof(Addr));
+	Addr.ss_family = AF_UNSPEC;
+}
+
+void FInternetAddrBSD::ResetScopeId()
+{
+#if PLATFORM_HAS_BSD_IPV6_SOCKETS
+	((sockaddr_in6*)&Addr)->sin6_scope_id = 0;
+#endif
+}
+
+void FInternetAddrBSD::SetIp(const sockaddr_storage& IpAddr)
+{
+	// Instead of just replacing the structures entirely, we copy only the ip portion.
+	// As this should not also set port.
+	if (IpAddr.ss_family == AF_INET)
+	{
+		const sockaddr_in* SockAddr = (const sockaddr_in*)&IpAddr;
+		SetIp(SockAddr->sin_addr);
+	}
+
+#if PLATFORM_HAS_BSD_IPV6_SOCKETS
+	if (IpAddr.ss_family == AF_INET6)
+	{
+		const sockaddr_in6* SockAddr = (const sockaddr_in6*)&IpAddr;
+		SetIp(SockAddr->sin6_addr);
+		// Remember to set the scope id.
+		((sockaddr_in6*)&Addr)->sin6_scope_id = SockAddr->sin6_scope_id;
+	}
+#endif
+}
+
+void FInternetAddrBSD::SetIp(const TCHAR* InAddr, bool& bIsValid)
+{
+	bIsValid = false;
+	FString AddressString(InAddr);
+	FString Port;
+
+	const bool bHasOpenBracket = AddressString.Contains("[");
+	const int32 CloseBracketIndex = AddressString.Find("]");
+	const bool bHasCloseBracket = CloseBracketIndex != INDEX_NONE;
+	bool bIsIPv6 = bHasOpenBracket && bHasCloseBracket;
+
+	const int32 LastColonIndex = AddressString.Find(":", ESearchCase::IgnoreCase, ESearchDir::FromEnd);
+
+	// IPv4 address will only have a port when a colon is present.
+	// IPv6 address will only have a port when surrounded by brackets.
+	const bool bHasPort = (INDEX_NONE != LastColonIndex) && (!bIsIPv6 || (bHasCloseBracket && LastColonIndex > CloseBracketIndex));
+
+	if (bHasPort)
+	{
+		Port = AddressString.RightChop(LastColonIndex + 1);
+		AddressString = AddressString.Left(LastColonIndex);
+	}
+
+	AddressString.RemoveFromStart("[");
+	AddressString.RemoveFromEnd("]");
+	
+	const auto InAddrAnsi = StringCast<ANSICHAR>(*AddressString);
+	FSocketSubsystemBSD* SocketSubsystem = static_cast<FSocketSubsystemBSD*>(ISocketSubsystem::Get());
+	if (SocketSubsystem)
+	{
+		bIsValid = (SocketSubsystem->CreateAddressFromIP(InAddrAnsi.Get(), *this) == SE_NO_ERROR);
+		if (bHasPort && bIsValid)
+		{
+			SetPort(FCString::Atoi(*Port));
+		}
+	}
+}
+
+TArray<uint8> FInternetAddrBSD::GetRawIp() const
+{
+	TArray<uint8> RawAddressArray;
+	if (Addr.ss_family == AF_INET)
+	{
+		sockaddr_in* IPv4Addr = ((sockaddr_in*)&Addr);
+		uint32 IntAddr = IPv4Addr->sin_addr.s_addr;
+		RawAddressArray.Add((IntAddr >> 24) & 0xFF);
+		RawAddressArray.Add((IntAddr >> 16) & 0xFF);
+		RawAddressArray.Add((IntAddr >> 8) & 0xFF);
+		RawAddressArray.Add((IntAddr >> 0) & 0xFF);
+	}
+#if PLATFORM_HAS_BSD_IPV6_SOCKETS
+	else if (Addr.ss_family == AF_INET6)
+	{
+		sockaddr_in6* IPv6Addr = ((sockaddr_in6*)&Addr);
+		for (int i = 0; i < 16; ++i)
+		{
+			RawAddressArray.Add(IPv6Addr->sin6_addr.s6_addr[i]);
+		}
+
+		// Copy over the interface if we have it explicitly set (anything other than 0)
+		if (IPv6Addr->sin6_scope_id != 0)
+		{
+			uint32 RawScopeId = IPv6Addr->sin6_scope_id;
+			RawAddressArray.Add((RawScopeId >> 24) & 0xFF);
+			RawAddressArray.Add((RawScopeId >> 16) & 0xFF);
+			RawAddressArray.Add((RawScopeId >> 8) & 0xFF);
+			RawAddressArray.Add((RawScopeId >> 0) & 0xFF);
+		}
+	}
+#endif
+
+	return RawAddressArray;
+}
+
+void FInternetAddrBSD::SetRawIp(const TArray<uint8>& RawAddr)
+{
+	if (RawAddr.Num() == 4) // This is IPv4
+	{
+		Addr.ss_family = AF_INET;
+		sockaddr_in* IPv4Addr = ((sockaddr_in*)&Addr);
+		IPv4Addr->sin_addr.s_addr = (RawAddr[0] << 24) | (RawAddr[1] << 16) | (RawAddr[2] << 8) | (RawAddr[3] << 0);
+	}
+#if PLATFORM_HAS_BSD_IPV6_SOCKETS
+	else if(RawAddr.Num() >= 16) // We are IPv6
+	{
+		sockaddr_in6* IPv6Addr = ((sockaddr_in6*)&Addr);
+		for (int i = 0; i < 16; ++i)
+		{
+			IPv6Addr->sin6_addr.s6_addr[i] = RawAddr[i];
+		}
+
+		// If this address has an interface, we'll copy it over too.
+		if (RawAddr.Num() == 20)
+		{
+			IPv6Addr->sin6_scope_id = (RawAddr[16] << 24) | (RawAddr[17] << 16) | (RawAddr[18] << 8) | (RawAddr[19] << 0);
+		}
+
+		Addr.ss_family = AF_INET6;
+	}
+#endif
+	else
+	{
+		Clear();
+	}
+}
+
+void FInternetAddrBSD::SetPort(int32 InPort)
+{
+#if PLATFORM_HAS_BSD_IPV6_SOCKETS
+	if (GetProtocolFamily() == ESocketProtocolFamily::IPv6)
+	{
+		((sockaddr_in6*)&Addr)->sin6_port = htons(InPort);
+		return;
+	}
+#endif
+
+	((sockaddr_in*)&Addr)->sin_port = htons(InPort);
+}
+
+int32 FInternetAddrBSD::GetPort() const
+{
+#if PLATFORM_HAS_BSD_IPV6_SOCKETS
+	if (GetProtocolFamily() == ESocketProtocolFamily::IPv6)
+	{
+		return ntohs(((sockaddr_in6*)&Addr)->sin6_port);
+	}
+#endif
+
+	return ntohs(((sockaddr_in*)&Addr)->sin_port);
+}
+
+void FInternetAddrBSD::SetAnyAddress()
+{
+#if PLATFORM_HAS_BSD_IPV6_SOCKETS
+	SetAnyIPv6Address();
+#else
+	SetAnyIPv4Address();
+#endif
+}
+
+void FInternetAddrBSD::SetAnyIPv4Address()
+{
+	((sockaddr_in*)&Addr)->sin_addr.s_addr = htonl(INADDR_ANY);
+	Addr.ss_family = AF_INET;
+	SetPort(0);
+}
+
+void FInternetAddrBSD::SetAnyIPv6Address()
+{
+#if PLATFORM_HAS_BSD_IPV6_SOCKETS
+	SetIp(in6addr_any);
+#endif
+	ResetScopeId();
+	SetPort(0);
+}
+
+void FInternetAddrBSD::SetBroadcastAddress()
+{
+#if PLATFORM_HAS_BSD_IPV6_SOCKETS
+	SetIPv6BroadcastAddress();
+#else
+	SetIPv4BroadcastAddress();
+#endif
+}
+
+void FInternetAddrBSD::SetIPv4BroadcastAddress()
+{
+	((sockaddr_in*)&Addr)->sin_addr.s_addr = htonl(INADDR_BROADCAST);
+	Addr.ss_family = AF_INET;
+	SetPort(0);
+}
+
+void FInternetAddrBSD::SetIPv6BroadcastAddress()
+{
+#if PLATFORM_HAS_BSD_IPV6_SOCKETS
+	// broadcast means something different in IPv6, but this is a rough equivalent
+#ifndef in6addr_allnodesonlink
+		// see RFC 4291, link-local multicast address http://tools.ietf.org/html/rfc4291
+	static in6_addr in6addr_allnodesonlink =
+	{
+		{ { 0xff, 0x02, 0x00, 0x00, 0x00, 0x00, 0x00, 0x00, 0x00, 0x00, 0x00, 0x00, 0x00, 0x00, 0x00, 0x01 } }
+	};
+#endif // in6addr_allnodesonlink
+	SetIp(in6addr_allnodesonlink);
+#endif
+	ResetScopeId();
+	SetPort(0);
+}
+
+void FInternetAddrBSD::SetLoopbackAddress()
+{
+#if PLATFORM_HAS_BSD_IPV6_SOCKETS
+	SetIPv6LoopbackAddress();
+#else
+	SetIPv4LoopbackAddress();
+#endif
+}
+
+void FInternetAddrBSD::SetIPv4LoopbackAddress()
+{
+	((sockaddr_in*)&Addr)->sin_addr.s_addr = htonl(INADDR_LOOPBACK);
+	Addr.ss_family = AF_INET;
+	SetPort(0);
+}
+
+void FInternetAddrBSD::SetIPv6LoopbackAddress()
+{
+#if PLATFORM_HAS_BSD_IPV6_SOCKETS
+	SetIp(in6addr_loopback);
+#endif
+	ResetScopeId();
+	SetPort(0);
+}
+
+FString FInternetAddrBSD::ToString(bool bAppendPort) const
+{
+	FString ReturnVal("");
+#if PLATFORM_HAS_BSD_SOCKET_FEATURE_GETNAMEINFO
+	char IPStr[NI_MAXHOST];
+	if (getnameinfo((const sockaddr*)&Addr, GetStorageSize(), IPStr, NI_MAXHOST, nullptr, 0, NI_NUMERICHOST) == 0)
+	{
+		if (GetProtocolFamily() == ESocketProtocolFamily::IPv6)
+		{
+			FString IPv6Str(ANSI_TO_TCHAR(IPStr));
+			// Remove the scope interface if it exists.
+			const int32 InterfaceMarkerIndex = IPv6Str.Find("%", ESearchCase::IgnoreCase, ESearchDir::FromEnd);
+			if (InterfaceMarkerIndex != INDEX_NONE)
+			{
+				IPv6Str = IPv6Str.Left(InterfaceMarkerIndex);
+			}
+
+			ReturnVal = FString::Printf(TEXT("[%s]"), *IPv6Str);
+		}
+		else
+		{
+			ReturnVal = IPStr;
+		}
+
+		if (bAppendPort)
+		{
+			ReturnVal += FString::Printf(TEXT(":%d"), GetPort());
+		}
+	}
+#else
+	UE_LOG(LogSockets, Error, TEXT("Platform has no getnameinfo(), but did not override FInternetAddrBSD::ToString()"));
+#endif
+	return ReturnVal;
+}
+
+bool FInternetAddrBSD::operator==(const FInternetAddr& Other) const
+{
+	const FInternetAddrBSD& OtherBSD = static_cast<const FInternetAddrBSD&>(Other);
+	ESocketProtocolFamily CurrentFamily = GetProtocolFamily();
+
+	// Check if the addr families match
+	if (OtherBSD.Addr.ss_family != Addr.ss_family)
+	{
+		return false;
+	}
+
+#if PLATFORM_HAS_BSD_IPV6_SOCKETS
+	if (CurrentFamily == ESocketProtocolFamily::IPv6)
+	{
+		const sockaddr_in6* OtherBSDAddr = (sockaddr_in6*)&(OtherBSD.Addr);
+		const sockaddr_in6* ThisBSDAddr = ((sockaddr_in6*)&Addr);
+		return memcmp(&(ThisBSDAddr->sin6_addr), &(OtherBSDAddr->sin6_addr), sizeof(in6_addr)) == 0 &&
+			ThisBSDAddr->sin6_port == OtherBSDAddr->sin6_port;
+	}
+#endif
+
+	if (CurrentFamily == ESocketProtocolFamily::IPv4)
+	{
+		const sockaddr_in* OtherBSDAddr = (sockaddr_in*)&(OtherBSD.Addr);
+		const sockaddr_in* ThisBSDAddr = ((sockaddr_in*)&Addr);
+		return ThisBSDAddr->sin_addr.s_addr == OtherBSDAddr->sin_addr.s_addr &&
+			ThisBSDAddr->sin_port == OtherBSDAddr->sin_port;
+	}
+
+	return false;
+}
+
+bool FInternetAddrBSD::IsValid() const
+{
+	ESocketProtocolFamily CurrentFamily = GetProtocolFamily();
+
+	if (CurrentFamily == ESocketProtocolFamily::IPv4)
+	{
+		return ((sockaddr_in*)&Addr)->sin_addr.s_addr != 0;
+	}
+
+#if PLATFORM_HAS_BSD_IPV6_SOCKETS
+	if (CurrentFamily == ESocketProtocolFamily::IPv6)
+	{
+		in6_addr EmptyAddr;
+		FMemory::Memzero(EmptyAddr);
+
+		sockaddr_in6* ThisBSDAddr = ((sockaddr_in6*)&Addr);
+		return memcmp(&(ThisBSDAddr->sin6_addr), &EmptyAddr, sizeof(in6_addr)) != 0;
+	}
+#endif
+
+	return false;
+}
+
+TSharedRef<FInternetAddr> FInternetAddrBSD::Clone() const
+{
+	TSharedRef<FInternetAddrBSD> NewAddress = MakeShareable(new FInternetAddrBSD);
+	NewAddress->Addr = Addr;
+	return NewAddress;
+}
+
+ESocketProtocolFamily FInternetAddrBSD::GetProtocolFamily() const
+{
+	switch (Addr.ss_family)
+	{
+		case AF_INET:
+			return ESocketProtocolFamily::IPv4;
+		break;
+		case AF_INET6:
+			return ESocketProtocolFamily::IPv6;
+		break;
+		default:
+			return ESocketProtocolFamily::None;
+	}
+}
+
+SOCKLEN FInternetAddrBSD::GetStorageSize() const
+{
+	if (GetProtocolFamily() == ESocketProtocolFamily::IPv4)
+	{
+		return sizeof(sockaddr_in);
+	}
+	else
+	{
+#if PLATFORM_HAS_BSD_IPV6_SOCKETS
+		return sizeof(sockaddr_in6);
+#endif
+	}
+
+	return sizeof(sockaddr_storage);
+}
+
+uint32 FInternetAddrBSD::GetTypeHash()
+{
+	ESocketProtocolFamily CurrentFamily = GetProtocolFamily();
+
+	if (CurrentFamily == ESocketProtocolFamily::IPv4)
+	{
+		uint32 NumericAddress;
+		GetIp(NumericAddress);
+		return NumericAddress + (GetPort() * 23);
+	}
+	else if (CurrentFamily == ESocketProtocolFamily::IPv6)
+	{
+		return ::GetTypeHash(*ToString(true));
+	}
+
+	return 0;
+}
+
+#endif
diff --git a/Engine/Source/Runtime/Sockets/Private/BSDSockets/IPAddressBSD.h b/Engine/Source/Runtime/Sockets/Private/BSDSockets/IPAddressBSD.h
--- a/Engine/Source/Runtime/Sockets/Private/BSDSockets/IPAddressBSD.h
+++ b/Engine/Source/Runtime/Sockets/Private/BSDSockets/IPAddressBSD.h
@@ -6,27 +6,65 @@
 #include "BSDSockets/SocketSubsystemBSDPrivate.h"
 #include "IPAddress.h"
 
-#if PLATFORM_HAS_BSD_SOCKETS
-
+#if PLATFORM_HAS_BSD_SOCKETS || PLATFORM_HAS_BSD_IPV6_SOCKETS
 
 /**
- * Represents an internet ip address, using the relatively standard SOCKADDR_IN structure. All data is in network byte order
+ * Represents an internet ip address with support for ipv4/v6. 
+ * All data is in network byte order
  */
 class FInternetAddrBSD : public FInternetAddr
 {
-protected:
 	/** The internet ip address structure */
-	sockaddr_in Addr;
+	sockaddr_storage Addr;
 
-public:
+protected:
 	/**
-	 * Constructor. Sets address to default state
+	 * Sets the ip address using a network byte order ipv4 address
+	 *
+	 * @param IPv4Addr the new ip address to use
+	 */
+	void SetIp(const in_addr& IPv4Addr)
+	{
+		((sockaddr_in*)&Addr)->sin_addr = IPv4Addr;
+		Addr.ss_family = AF_INET;
+	}
+
+#if PLATFORM_HAS_BSD_IPV6_SOCKETS
+	/**
+	 * Sets the ip address using a network byte order ipv6 address
+	 *
+	 * @param IpAddr the new ip address to use
 	 */
-	FInternetAddrBSD()
+	void SetIp(const in6_addr& IpAddr)
 	{
-		FMemory::Memzero(&Addr,sizeof(sockaddr_in));
-		Addr.sin_family = AF_INET;
+		((sockaddr_in6*)&Addr)->sin6_addr = IpAddr;
+		Addr.ss_family = AF_INET6;
 	}
+#endif
+
+	void Clear();
+	void ResetScopeId();
+
+public:
+	/**
+	 * Constructor. Sets address to default state
+	 */
+	FInternetAddrBSD();
+
+	/**
+	 * Sets the ip address from a raw network byte order array.
+	 *
+	 * @param RawAddr the new address to use (must be converted to network byte order)
+	 */
+	virtual void SetRawIp(const TArray<uint8>& RawAddr) override;
+
+	/**
+	 * Gets the ip address in a raw array stored in network byte order.
+	 *
+	 * @return The raw address stored in an array
+	 */
+	virtual TArray<uint8> GetRawIp() const override;
+
 	/**
 	 * Sets the ip address from a host byte order uint32
 	 *
@@ -34,189 +72,186 @@ public:
 	 */
 	virtual void SetIp(uint32 InAddr) override
 	{
-		Addr.sin_addr.s_addr = htonl(InAddr);
+		((sockaddr_in*)&Addr)->sin_addr.s_addr = htonl(InAddr);
+		Addr.ss_family = AF_INET;
 	}
 
 	/**
-	 * Sets the ip address from a string ("A.B.C.D")
+	 * Sets the ip address from a string IPv6 or IPv4 address.
+	 * Ports may be included using the form Address:Port, or excluded and set manually.
+	 *
+	 * IPv6 - [1111:2222:3333:4444:5555:6666:7777:8888]:PORT || [1111:2222:3333::]:PORT || [::ffff:IPv4]:PORT
+	 * IPv4 - aaa.bbb.ccc.ddd:PORT || 127.0.0.1:PORT
 	 *
 	 * @param InAddr the string containing the new ip address to use
+	 * @param bIsValid will be set to true if InAddr was a valid IPv6 or IPv4 address, false if not.
+	 */
+	virtual void SetIp(const TCHAR* InAddr, bool& bIsValid) override;
+
+	/**
+	 * Sets the ip address using a generic sockaddr_storage
+	 *
+	 * @param IpAddr the new ip address to use (assumes already in the correct byte order).
+	 */
+	void SetIp(const sockaddr_storage& IpAddr);
+
+	/**
+	 * Sets the address data via a sockaddr_storage
+	 *
+	 * @param AddrData the new data to use
 	 */
-	virtual void SetIp(const TCHAR* InAddr, bool& bIsValid) override
+	virtual void Set(const sockaddr_storage& AddrData)
 	{
-		int32 Port = 0;
+		Addr = AddrData;
+	}
 
-		FString AddressString = InAddr;
-		if (AddressString.IsEmpty())
+#if PLATFORM_HAS_BSD_IPV6_SOCKETS
+	/**
+	 * Copies the network byte order ip address
+	 *
+	 * @param OutAddr the out param receiving the ip address
+	 */
+	void GetIp(in6_addr& OutAddr) const
+	{
+		if (GetProtocolFamily() != ESocketProtocolFamily::IPv6)
 		{
-			bIsValid = false;
 			return;
 		}
 
-		TArray<FString> PortTokens;
-		AddressString.ParseIntoArray(PortTokens, TEXT(":"), true);
-		
-		// look for a port number
-		if (PortTokens.Num() > 1)
-		{
-			Port = FCString::Atoi(*PortTokens[1]);
-		}
-
-		// Check if it's a valid IPv4 address, and if it is convert
-		in_addr IPv4Addr;
-		const auto InAddrAnsi = StringCast<ANSICHAR>(*(PortTokens[0]));
-		if (inet_pton(AF_INET, InAddrAnsi.Get(), &IPv4Addr))
-		{
-			if (Port != 0)
-			{
-				SetPort(Port);
-			}
-
-			bIsValid = true;
-
-			SetIp(IPv4Addr);
-		}
-		else
-		{
-			//debugf(TEXT("Invalid IP address string (%s) passed to SetIp"),InAddr);
-			bIsValid = false;
-		}
+		OutAddr = ((sockaddr_in6*)&Addr)->sin6_addr;
 	}
+#endif
 
 	/**
-	 * Sets the ip address using a network byte order ip address
+	 * Copies the network byte order ip address
 	 *
-	 * @param IpAddr the new ip address to use
+	 * @param OutAddr the out param receiving the ip address
 	 */
-	void SetIp(const in_addr& IpAddr)
- 	{
- 		Addr.sin_addr = IpAddr;
- 	}
+	void GetIp(in_addr& OutAddr) const
+	{
+		if (GetProtocolFamily() != ESocketProtocolFamily::IPv4)
+		{
+			return;
+		}
 
+		OutAddr = ((sockaddr_in*)&Addr)->sin_addr;
+	}
 
 	/**
-	 * Copies the network byte order ip address to a host byte order dword
+	 * Copies the network byte order ip address to a host byte order dword.
+	 * Does nothing if we are currently not storing an ipv4 addr
 	 *
 	 * @param OutAddr the out param receiving the ip address
 	 */
 	virtual void GetIp(uint32& OutAddr) const override
-	{ 
-		OutAddr = ntohl(Addr.sin_addr.s_addr);
+	{
+		if (GetProtocolFamily() != ESocketProtocolFamily::IPv4)
+		{
+			OutAddr = 0;
+			return;
+		}
+
+		OutAddr = ntohl(((sockaddr_in*)&Addr)->sin_addr.s_addr);
 	}
 
 	/**
-	 * Copies the network byte order ip address 
+	 * Sets the port number from a host byte order int
 	 *
-	 * @param OutAddr the out param receiving the ip address
+	 * @param InPort the new port to use (must convert to network byte order)
 	 */
-	void GetIp(in_addr& OutAddr) const
- 	{
- 		OutAddr = Addr.sin_addr;
- 	}
+	virtual void SetPort(int32 InPort) override;
+
+	/** Returns the port number from this address in host byte order */
+	virtual int32 GetPort() const override;
 
 	/**
-	 * Sets the port number from a host byte order int
+	 * Sets the address structure to be bindable to any IP address.
+	 * IPv6 will take precedence.
 	 *
-	 * @param InPort the new port to use (must convert to network byte order)
+	 * To skip assumptions, you can call the designated version explicitly below.
 	 */
-	virtual void SetPort(int32 InPort) override
-	{
-		Addr.sin_port = htons(InPort);
-	}
+	virtual void SetAnyAddress() override;
+
+	/** Explicit set to any IPv4 address */
+	void SetAnyIPv4Address();
+
+	/** Explicit set to any IPv6 address */
+	void SetAnyIPv6Address();
 
 	/**
-	 * Copies the port number from this address and places it into a host byte order int
+	 * Sets the address structure to be bound to the multicast ip address.
+	 * IPv6 will take precedence.
 	 *
-	 * @param OutPort the host byte order int that receives the port
+	 * To skip assumptions, you can call the designated version explicitly below.
 	 */
-	virtual void GetPort(int32& OutPort) const override
-	{
-		OutPort = ntohs(Addr.sin_port);
-	}
+	virtual void SetBroadcastAddress() override;
+
+	/** Explicit set to multicast IPv4 address */
+	void SetIPv4BroadcastAddress();
+
+	/** Explicit set to multicast IPv6 address */
+	void SetIPv6BroadcastAddress();
 
 	/**
-	 * Returns the port number from this address in host byte order
+	 * Sets the address structure to be bound to the loopback ip address.
+	 * IPv6 will take precedence.
+	 *
+	 * To skip assumptions, you can call the designated version explicitly below.
 	 */
-	virtual int32 GetPort(void) const override
-	{
-		return ntohs(Addr.sin_port);
-	}
+	virtual void SetLoopbackAddress() override;
 
-	/** Sets the address to be any address */
-	virtual void SetAnyAddress(void) override
-	{
-		SetIp(INADDR_ANY);
-		SetPort(0);
-	}
+	/** Explicit set to loopback IPv4 address */
+	void SetIPv4LoopbackAddress();
 
-	/** Sets the address to broadcast */
-	virtual void SetBroadcastAddress() override
-	{
-		SetIp(INADDR_BROADCAST);
-		SetPort(0);
-	}
+	/** Explicit set to loopback IPv6 address */
+	void SetIPv6LoopbackAddress();
 
 	/**
-	 * Converts this internet ip address to string form
+	 * Converts this internet ip address to string form. String will be enclosed in square braces.
 	 *
 	 * @param bAppendPort whether to append the port information or not
 	 */
-	virtual FString ToString(bool bAppendPort) const override
-	{
-		uint32 LocalAddr = ntohl(Addr.sin_addr.s_addr);
-
-		// Get the individual bytes
-		const int32 A = (LocalAddr >> 24) & 0xFF;
-		const int32 B = (LocalAddr >> 16) & 0xFF;
-		const int32 C = (LocalAddr >>  8) & 0xFF;
-		const int32 D = (LocalAddr >>  0) & 0xFF;
-		if (bAppendPort)
-		{
-			return FString::Printf(TEXT("%i.%i.%i.%i:%i"),A,B,C,D,GetPort());
-		}
-		else
-		{
-			return FString::Printf(TEXT("%i.%i.%i.%i"),A,B,C,D);
-		}
-	}
+	virtual FString ToString(bool bAppendPort) const override;
 
 	/**
 	 * Compares two internet ip addresses for equality
 	 *
 	 * @param Other the address to compare against
 	 */
-	virtual bool operator==(const FInternetAddr& Other) const override
-	{
-		FInternetAddrBSD& OtherBSD = (FInternetAddrBSD&)Other;
-		return Addr.sin_addr.s_addr == OtherBSD.Addr.sin_addr.s_addr &&
-			Addr.sin_port == OtherBSD.Addr.sin_port &&
-			Addr.sin_family == OtherBSD.Addr.sin_family;
-	}
-
-	virtual uint32 GetTypeHash() override
-	{
-		return Addr.sin_addr.s_addr + (Addr.sin_port * 23);
-	}
+	virtual bool operator==(const FInternetAddr& Other) const override;
 
 	/**
-	 * Is this a well formed internet address
+	 * Is this a well formed internet address, the only criteria being non-zero
 	 *
 	 * @return true if a valid IP, false otherwise
 	 */
-	virtual bool IsValid() const override
-	{
-		return Addr.sin_addr.s_addr != 0;
-	}
+	virtual bool IsValid() const override;
 
- 	operator sockaddr*(void)
- 	{
- 		return (sockaddr*)&Addr;
- 	}
+	/**
+	 * Creates a new structure with the same data as this structure
+	 *
+	 * @return The new structure
+	 */
+	virtual TSharedRef<FInternetAddr> Clone() const override;
+
+	/**
+	 * Returns the protocol family of the address data currently stored in this struct
+	 *
+	 * @return They type of the address
+	 */
+	virtual ESocketProtocolFamily GetProtocolFamily() const;
+
+	/**
+	 * Returns the size of the amount of data that is being used
+	 * to hold the address information. Useful for functions like bind/connect
+	 *
+	 * @return size of addr
+	 */
+	virtual SOCKLEN GetStorageSize() const;
+	
+	virtual uint32 GetTypeHash() override;
 
- 	operator const sockaddr*(void) const
- 	{
- 		return (const sockaddr*)&Addr;
- 	}
+	friend class FSocketBSD;
 };
 
-#endif	//PLATFORM_HAS_BSD_SOCKETS
+#endif
diff --git a/Engine/Source/Runtime/Sockets/Private/BSDSockets/SocketSubsystemBSD.cpp b/Engine/Source/Runtime/Sockets/Private/BSDSockets/SocketSubsystemBSD.cpp
--- a/Engine/Source/Runtime/Sockets/Private/BSDSockets/SocketSubsystemBSD.cpp
+++ b/Engine/Source/Runtime/Sockets/Private/BSDSockets/SocketSubsystemBSD.cpp
@@ -1,9 +1,8 @@
 // Copyright 1998-2018 Epic Games, Inc. All Rights Reserved.
 
 #include "BSDSockets/SocketSubsystemBSD.h"
-#include "Misc/ScopeLock.h"
 
-#if PLATFORM_HAS_BSD_SOCKETS
+#if PLATFORM_HAS_BSD_SOCKETS || PLATFORM_HAS_BSD_IPV6_SOCKETS
 
 #include "IPAddress.h"
 #include "BSDSockets/IPAddressBSD.h"
@@ -16,34 +15,32 @@ FSocketBSD* FSocketSubsystemBSD::InternalBSDSocketFactory(SOCKET Socket, ESocket
 	return new FSocketBSD(Socket, SocketType, SocketDescription, this);
 }
 
-
 FSocket* FSocketSubsystemBSD::CreateSocket(const FName& SocketType, const FString& SocketDescription, bool bForceUDP)
+{
+	return CreateSocket(SocketType, SocketDescription, GetDefaultSocketProtocolFamily(), bForceUDP);
+}
+
+class FSocket* FSocketSubsystemBSD::CreateSocket(const FName& SocketType, const FString& SocketDescription, ESocketProtocolFamily ProtocolType, bool bForceUDP)
 {
 	SOCKET Socket = INVALID_SOCKET;
 	FSocket* NewSocket = nullptr;
 	int PlatformSpecificTypeFlags = 0;
 
+	// Don't support any other protocol families.
+	if (ProtocolType != ESocketProtocolFamily::IPv4 && ProtocolType != ESocketProtocolFamily::IPv6)
+	{
+		return nullptr;
+	}
+
 #if PLATFORM_HAS_BSD_SOCKET_FEATURE_CLOSE_ON_EXEC
 	PlatformSpecificTypeFlags = SOCK_CLOEXEC;
 #endif // PLATFORM_HAS_BSD_SOCKET_FEATURE_CLOSE_ON_EXEC
 
-	switch (SocketType.GetComparisonIndex())
-	{
-	case NAME_DGram:
-		// Creates a data gram (UDP) socket
-		Socket = socket(AF_INET, SOCK_DGRAM | PlatformSpecificTypeFlags, IPPROTO_UDP);
-		NewSocket = (Socket != INVALID_SOCKET) ? InternalBSDSocketFactory(Socket, SOCKTYPE_Datagram, SocketDescription) : nullptr;
-		break;
-
-	case NAME_Stream:
-		// Creates a stream (TCP) socket
-		Socket = socket(AF_INET, SOCK_STREAM | PlatformSpecificTypeFlags, IPPROTO_TCP);
-		NewSocket = (Socket != INVALID_SOCKET) ? InternalBSDSocketFactory(Socket, SOCKTYPE_Streaming, SocketDescription) : nullptr;
-		break;
+	bool bIsUDP = (bForceUDP || SocketType.GetComparisonIndex() == NAME_DGram);
+	int32 SocketTypeFlag = (bIsUDP) ? SOCK_DGRAM : SOCK_STREAM;
 
-	default:
-		break;
-	}
+	Socket = socket(GetProtocolFamilyValue(ProtocolType), SocketTypeFlag | PlatformSpecificTypeFlags, ((bIsUDP) ? IPPROTO_UDP : IPPROTO_TCP));
+	NewSocket = (Socket != INVALID_SOCKET) ? InternalBSDSocketFactory(Socket, ((bIsUDP) ? SOCKTYPE_Datagram : SOCKTYPE_Streaming), SocketDescription) : nullptr;
 
 	if (!NewSocket)
 	{
@@ -53,23 +50,55 @@ FSocket* FSocketSubsystemBSD::CreateSocket(const FName& SocketType, const FStrin
 	return NewSocket;
 }
 
+ESocketErrors FSocketSubsystemBSD::CreateAddressFromIP(const ANSICHAR* IPAddress, FInternetAddr& OutAddr)
+{
+	TArray<TSharedRef<FInternetAddr> > Results = GetAddrInfo(IPAddress, false);
+	if (Results.Num() > 0)
+	{
+		OutAddr.SetRawIp(Results[0]->GetRawIp());
+		return SE_NO_ERROR;
+	}
+	
+	return SE_HOST_NOT_FOUND;
+}
+
+ESocketErrors FSocketSubsystemBSD::GetHostByName(const ANSICHAR* HostName, FInternetAddr& OutAddr)
+{
+	TArray<TSharedRef<FInternetAddr> > Results = GetAddrInfo(HostName, true);
+	if (Results.Num() > 0)
+	{
+		OutAddr.SetRawIp(Results[0]->GetRawIp());
+		return SE_NO_ERROR;
+	}
+
+	return SE_HOST_NOT_FOUND;
+}
 
 void FSocketSubsystemBSD::DestroySocket(FSocket* Socket)
 {
 	delete Socket;
 }
 
-
-ESocketErrors FSocketSubsystemBSD::GetHostByName(const ANSICHAR* HostName, FInternetAddr& OutAddr)
+TArray<TSharedRef<FInternetAddr> > FSocketSubsystemBSD::GetAddrInfo(const ANSICHAR* HostName, bool bResolveAddress, ESocketProtocolFamily ProtocolType)
 {
+	TArray<TSharedRef<FInternetAddr> > GetAddrArray;
+
 #if PLATFORM_HAS_BSD_SOCKET_FEATURE_GETADDRINFO
-	FScopeLock ScopeLock(&HostByNameSynch);
 	addrinfo* AddrInfo = nullptr;
 
-	// Limit the IP Addresses we get back to just IPv4
+	// Make sure we filter out IPv6 if the platform is not officially supported
+	// (if it isn't supported but we explicitly ask for it, allow it).
+	bool bCanUseIPv6 = (PLATFORM_HAS_BSD_IPV6_SOCKETS || ProtocolType == ESocketProtocolFamily::IPv6) ? true : false;
+
 	addrinfo HintAddrInfo;
 	FMemory::Memzero(&HintAddrInfo, sizeof(HintAddrInfo));
-	HintAddrInfo.ai_family = AF_INET;
+	HintAddrInfo.ai_family = GetProtocolFamilyValue(ProtocolType);
+	HintAddrInfo.ai_flags = AI_ADDRCONFIG;
+
+	if (!bResolveAddress)
+	{
+		HintAddrInfo.ai_flags |= AI_NUMERICHOST;
+	}
 
 	int32 ErrorCode = getaddrinfo(HostName, nullptr, &HintAddrInfo, &AddrInfo);
 	ESocketErrors SocketError = TranslateGAIErrorCode(ErrorCode);
@@ -77,29 +106,29 @@ ESocketErrors FSocketSubsystemBSD::GetHostByName(const ANSICHAR* HostName, FInte
 	{
 		for (; AddrInfo != nullptr; AddrInfo = AddrInfo->ai_next)
 		{
-			if (AddrInfo->ai_family == AF_INET)
+			if (AddrInfo->ai_family == AF_INET || (AddrInfo->ai_family == AF_INET6 && bCanUseIPv6))
 			{
-				sockaddr_in* IPv4SockAddr = reinterpret_cast<sockaddr_in*>(AddrInfo->ai_addr);
-				if (IPv4SockAddr != nullptr)
+				sockaddr_storage* AddrData = reinterpret_cast<sockaddr_storage*>(AddrInfo->ai_addr);
+				if (AddrData != nullptr)
 				{
-					uint32 HostIP = ntohl(IPv4SockAddr->sin_addr.s_addr);
-					static_cast<FInternetAddrBSD&>(OutAddr).SetIp(HostIP);
-					freeaddrinfo(AddrInfo);
-					return SE_NO_ERROR;
+					// This might not be good as it locks GetAddrInfo into one class.
+					// TODO: Fix this when GetAddrInfo moves to the base class.
+					// Just convert the sockaddr_storage into the raw ip array format (do we want to make a new helper in this class?)
+					// and push it to SetRawIp
+					TSharedRef<FInternetAddrBSD> NewAddress = MakeShareable(new FInternetAddrBSD);
+					NewAddress->SetIp(*AddrData);
+					GetAddrArray.Add(NewAddress);
 				}
 			}
 		}
 		freeaddrinfo(AddrInfo);
-		return SE_HOST_NOT_FOUND;
 	}
-	return SocketError;
 #else
 	UE_LOG(LogSockets, Error, TEXT("Platform has no getaddrinfo(), but did not override FSocketSubsystem::GetHostByName()"));
-	return SE_NO_RECOVERY;
 #endif
+	return GetAddrArray;
 }
 
-
 bool FSocketSubsystemBSD::GetHostName(FString& HostName)
 {
 #if PLATFORM_HAS_BSD_SOCKET_FEATURE_GETHOSTNAME
@@ -119,19 +148,17 @@ bool FSocketSubsystemBSD::GetHostName(FString& HostName)
 
 const TCHAR* FSocketSubsystemBSD::GetSocketAPIName() const
 {
-	return TEXT("BSD");
+	return TEXT("BSD IPv4/6");
 }
 
-
 TSharedRef<FInternetAddr> FSocketSubsystemBSD::CreateInternetAddr(uint32 Address, uint32 Port)
 {
-	TSharedRef<FInternetAddr> Result = MakeShareable(new FInternetAddrBSD);
-	Result->SetIp(Address);
-	Result->SetPort(Port);
-	return Result;
+	TSharedRef<FInternetAddr> ReturnAddr = MakeShareable(new FInternetAddrBSD);
+	ReturnAddr->SetIp(Address);
+	ReturnAddr->SetPort(Port);
+	return ReturnAddr;
 }
 
-
 ESocketErrors FSocketSubsystemBSD::GetLastErrorCode()
 {
 	return TranslateErrorCode(errno);
@@ -219,5 +246,47 @@ ESocketErrors FSocketSubsystemBSD::TranslateErrorCode(int32 Code)
 	return SE_EINVAL;
 }
 
+ESocketErrors FSocketSubsystemBSD::TranslateGAIErrorCode(int32 Code) const
+{
+#if PLATFORM_HAS_BSD_SOCKET_FEATURE_GETADDRINFO
+	switch (Code)
+	{
+	// getaddrinfo() has its own error codes
+	case EAI_AGAIN:			return SE_TRY_AGAIN;
+	case EAI_BADFLAGS:		return SE_EINVAL;
+	case EAI_FAIL:			return SE_NO_RECOVERY;
+	case EAI_FAMILY:		return SE_EAFNOSUPPORT;
+	case EAI_MEMORY:		return SE_ENOBUFS;
+	case EAI_NONAME:		return SE_HOST_NOT_FOUND;
+	case EAI_SERVICE:		return SE_EPFNOSUPPORT;
+	case EAI_SOCKTYPE:		return SE_ESOCKTNOSUPPORT;
+#if PLATFORM_HAS_BSD_SOCKET_FEATURE_WINSOCKETS
+	case WSANO_DATA:		return SE_NO_DATA;
+	case WSANOTINITIALISED: return SE_NOTINITIALISED;
+#else			
+	case EAI_NODATA:		return SE_NO_DATA;
+	case EAI_ADDRFAMILY:	return SE_ADDRFAMILY;
+	case EAI_SYSTEM:		return SE_SYSTEM;
+#endif
+	case 0:					break; // 0 means success
+	default:
+		UE_LOG(LogSockets, Warning, TEXT("Unhandled getaddrinfo() socket error! Code: %d"), Code);
+		return SE_EINVAL;
+	}
+#endif // PLATFORM_HAS_BSD_SOCKET_FEATURE_GETADDRINFO
+
+	return SE_NO_ERROR;
+}
+
+int32 FSocketSubsystemBSD::GetProtocolFamilyValue(ESocketProtocolFamily InProtocol) const
+{
+	switch (InProtocol)
+	{
+		default:
+		case ESocketProtocolFamily::None: return AF_UNSPEC;
+		case ESocketProtocolFamily::IPv4: return AF_INET;
+		case ESocketProtocolFamily::IPv6: return AF_INET6;
+	}
+}
 
 #endif	//PLATFORM_HAS_BSD_SOCKETS
diff --git a/Engine/Source/Runtime/Sockets/Private/BSDSockets/SocketSubsystemBSD.h b/Engine/Source/Runtime/Sockets/Private/BSDSockets/SocketSubsystemBSD.h
--- a/Engine/Source/Runtime/Sockets/Private/BSDSockets/SocketSubsystemBSD.h
+++ b/Engine/Source/Runtime/Sockets/Private/BSDSockets/SocketSubsystemBSD.h
@@ -5,26 +5,57 @@
 #include "CoreMinimal.h"
 #include "Misc/OutputDeviceRedirector.h"
 #include "SocketSubsystemPackage.h"
+#include "SocketTypes.h"
+
+#if PLATFORM_HAS_BSD_SOCKETS || PLATFORM_HAS_BSD_IPV6_SOCKETS
 
 class FInternetAddr;
 class FSocketBSD;
 
-#if PLATFORM_HAS_BSD_SOCKETS
-
 #include "SocketSubsystemBSDPrivate.h"
 
-
 /**
  * Standard BSD specific socket subsystem implementation
  */
-class FSocketSubsystemBSD
-	: public FSocketSubsystemBSDCommon
+class FSocketSubsystemBSD : public ISocketSubsystem
 {
 public:
 
-	// ISocketSubsystem interface
+	virtual class FSocket* CreateSocket(const FName& SocketType, const FString& SocketDescription, ESocketProtocolFamily ProtocolType, bool bForceUDP = false);
+	virtual ESocketErrors CreateAddressFromIP(const ANSICHAR* IPAddress, FInternetAddr& OutAddr);
+
+	/**
+	 * Gets the address information of the given hostname and outputs it into an array of resolvable addresses.
+	 * It is up to the caller to determine which one is valid for their environment.
+	 *
+	 * @param HostName string version of the queryable hostname or ip address
+	 * @param bResolveAddress if a blocking query to the DNS provider should be make to determine the ip address of the given hostname.
+	 *						  if this flag is false, the call is asynchronous and will only work with ip addresses
+	 * @param ProtocolType this is used to limit results from the call. Specifying None will search all valid protocols.
+	 *					   Callers will find they rarely have to specify this flag.
+	 *
+	 * @return the array of results from GetAddrInfo
+	 */
+	virtual TArray<TSharedRef<FInternetAddr> > GetAddrInfo(const ANSICHAR* HostName, bool bResolveAddress, ESocketProtocolFamily ProtocolType = ESocketProtocolFamily::None);
+
+	/**
+	 * Specifies the default socket protocol family to use when creating a socket
+	 * without explicitly passing in the protocol type on creation.
+	 *
+	 * All sockets created using the base class's CreateSocket will use this function
+	 * to determine domain.
+	 */
+	virtual ESocketProtocolFamily GetDefaultSocketProtocolFamily() const
+	{
+#if PLATFORM_HAS_BSD_IPV6_SOCKETS
+		return ESocketProtocolFamily::IPv6;
+#else
+		return ESocketProtocolFamily::IPv4;
+#endif
+	}
 
-	virtual TSharedRef<FInternetAddr> CreateInternetAddr( uint32 Address = 0, uint32 Port = 0 ) override;
+	// ISocketSubsystem interface
+	virtual TSharedRef<FInternetAddr> CreateInternetAddr(uint32 Address = 0, uint32 Port = 0) override;
 	virtual class FSocket* CreateSocket( const FName& SocketType, const FString& SocketDescription, bool bForceUDP = false ) override;
 	virtual void DestroySocket( class FSocket* Socket ) override;
 	virtual ESocketErrors GetHostByName( const ANSICHAR* HostName, FInternetAddr& OutAddr ) override;
@@ -55,6 +86,15 @@ public:
 	virtual ESocketErrors TranslateErrorCode( int32 Code ) override;
 
 protected:
+	/**
+	 * Translates return values of getaddrinfo() to socket error enum
+	 */
+	ESocketErrors TranslateGAIErrorCode(int32 Code) const;
+
+	/**
+	 * Translates an ESocketProtocolFamily code into a value usable by raw socket apis.
+	 */
+	virtual int32 GetProtocolFamilyValue(ESocketProtocolFamily InProtocol) const;
 
 	/**
 	 * Allows a subsystem subclass to create a FSocketBSD sub class.
@@ -63,12 +103,6 @@ protected:
 
 	// allow BSD sockets to use this when creating new sockets from accept() etc
 	friend FSocketBSD;
-
-private:
-
-	// Used to prevent multiple threads accessing the shared data.
-	FCriticalSection HostByNameSynch;
 };
 
-
 #endif
diff --git a/Engine/Source/Runtime/Sockets/Private/BSDSockets/SocketSubsystemBSDPrivate.h b/Engine/Source/Runtime/Sockets/Private/BSDSockets/SocketSubsystemBSDPrivate.h
--- a/Engine/Source/Runtime/Sockets/Private/BSDSockets/SocketSubsystemBSDPrivate.h
+++ b/Engine/Source/Runtime/Sockets/Private/BSDSockets/SocketSubsystemBSDPrivate.h
@@ -75,7 +75,7 @@ inline int TranslateFlags(ESocketReceiveFlags::Type Flags)
 /**
  * Standard BSD specific socket subsystem implementation (common to both IPv4 and IPv6)
  */
-class FSocketSubsystemBSDCommon
+class DEPRECATED(4.21, "Use FSocketSubsystemBSD.") FSocketSubsystemBSDCommon
 	: public ISocketSubsystem
 {
 protected:
diff --git a/Engine/Source/Runtime/Sockets/Private/BSDSockets/SocketsBSD.cpp b/Engine/Source/Runtime/Sockets/Private/BSDSockets/SocketsBSD.cpp
--- a/Engine/Source/Runtime/Sockets/Private/BSDSockets/SocketsBSD.cpp
+++ b/Engine/Source/Runtime/Sockets/Private/BSDSockets/SocketsBSD.cpp
@@ -2,7 +2,7 @@
 
 #include "BSDSockets/SocketsBSD.h"
 
-#if PLATFORM_HAS_BSD_SOCKETS
+#if PLATFORM_HAS_BSD_SOCKETS || PLATFORM_HAS_BSD_IPV6_SOCKETS
 
 #include "BSDSockets/IPAddressBSD.h"
 #include "BSDSockets/SocketSubsystemBSD.h"
@@ -30,13 +30,15 @@ bool FSocketBSD::Close(void)
 
 bool FSocketBSD::Bind(const FInternetAddr& Addr)
 {
-	return bind(Socket, (sockaddr*)(FInternetAddrBSD&)Addr, sizeof(sockaddr_in)) == 0;
+	const FInternetAddrBSD& BSDAddr = static_cast<const FInternetAddrBSD&>(Addr);
+	return bind(Socket, (sockaddr*)&BSDAddr.Addr, BSDAddr.GetStorageSize()) == 0;
 }
 
 
 bool FSocketBSD::Connect(const FInternetAddr& Addr)
 {
-	int32 Return = connect(Socket, (sockaddr*)(FInternetAddrBSD&)Addr, sizeof(sockaddr_in));
+	const FInternetAddrBSD& BSDAddr = static_cast<const FInternetAddrBSD&>(Addr);
+	int32 Return = connect(Socket, (sockaddr*)&BSDAddr.Addr, BSDAddr.GetStorageSize());
 	
 	check(SocketSubsystem);
 	ESocketErrors Error = SocketSubsystem->TranslateErrorCode(Return);
@@ -110,8 +112,9 @@ FSocket* FSocketBSD::Accept(const FString& InSocketDescription)
 
 FSocket* FSocketBSD::Accept(FInternetAddr& OutAddr, const FString& InSocketDescription)
 {
-	SOCKLEN SizeOf = sizeof(sockaddr_in);
-	SOCKET NewSocket = accept(Socket, *(FInternetAddrBSD*)(&OutAddr), &SizeOf);
+	FInternetAddrBSD& BSDAddr = static_cast<FInternetAddrBSD&>(OutAddr);
+	SOCKLEN SizeOf = sizeof(sockaddr_storage);
+	SOCKET NewSocket = accept(Socket, (sockaddr*)&(BSDAddr.Addr), &SizeOf);
 
 	if (NewSocket != INVALID_SOCKET)
 	{
@@ -127,8 +130,9 @@ FSocket* FSocketBSD::Accept(FInternetAddr& OutAddr, const FString& InSocketDescr
 
 bool FSocketBSD::SendTo(const uint8* Data, int32 Count, int32& BytesSent, const FInternetAddr& Destination)
 {
+	const FInternetAddrBSD& BSDAddr = static_cast<const FInternetAddrBSD&>(Destination);
 	// Write the data and see how much was written
-	BytesSent = sendto(Socket, (const char*)Data, Count, 0, (FInternetAddrBSD&)Destination, sizeof(sockaddr_in));
+	BytesSent = sendto(Socket, (const char*)Data, Count, 0, (sockaddr*)&(BSDAddr.Addr), BSDAddr.GetStorageSize());
 
 //	NETWORK_PROFILER(FSocket::SendTo(Data,Count,BytesSent,Destination));
 
@@ -161,11 +165,12 @@ bool FSocketBSD::RecvFrom(uint8* Data, int32 BufferSize, int32& BytesRead, FInte
 	bool bSuccess = false;
 	const bool bStreamSocket = (SocketType == SOCKTYPE_Streaming);
 	const int TranslatedFlags = TranslateFlags(Flags);
-	SOCKLEN Size = sizeof(sockaddr_in);
-	sockaddr& Addr = *(FInternetAddrBSD&)Source;
+	FInternetAddrBSD& BSDAddr = static_cast<FInternetAddrBSD&>(Source);
+	SOCKLEN Size = sizeof(sockaddr_storage);
+	sockaddr* Addr = (sockaddr*)&(BSDAddr.Addr);
 
 	// Read into the buffer and set the source address
-	BytesRead = recvfrom(Socket, (char*)Data, BufferSize, TranslatedFlags, &Addr, &Size);
+	BytesRead = recvfrom(Socket, (char*)Data, BufferSize, TranslatedFlags, Addr, &Size);
 //	NETWORK_PROFILER(FSocket::RecvFrom(Data,BufferSize,BytesRead,Source));
 
 	if (BytesRead >= 0)
@@ -277,11 +282,11 @@ ESocketConnectionState FSocketBSD::GetConnectionState(void)
 
 void FSocketBSD::GetAddress(FInternetAddr& OutAddr)
 {
-	FInternetAddrBSD& Addr = (FInternetAddrBSD&)OutAddr;
-	SOCKLEN Size = sizeof(sockaddr_in);
+	FInternetAddrBSD& BSDAddr = static_cast<FInternetAddrBSD&>(OutAddr);
+	SOCKLEN Size = sizeof(sockaddr_storage);
 
 	// Figure out what ip/port we are bound to
-	bool bOk = getsockname(Socket, Addr, &Size) == 0;
+	bool bOk = getsockname(Socket, (sockaddr*)&BSDAddr.Addr, &Size) == 0;
 
 	if (bOk == false)
 	{
@@ -293,11 +298,11 @@ void FSocketBSD::GetAddress(FInternetAddr& OutAddr)
 
 bool FSocketBSD::GetPeerAddress(FInternetAddr& OutAddr)
 {
-	FInternetAddrBSD& Addr = (FInternetAddrBSD&)OutAddr;
-	SOCKLEN Size = sizeof(sockaddr_in);
+	FInternetAddrBSD& BSDAddr = static_cast<FInternetAddrBSD&>(OutAddr);
+	SOCKLEN Size = sizeof(sockaddr_storage);
 
 	// Figure out what ip/port we are bound to
-	int Result = getpeername(Socket, Addr, &Size);
+	int Result = getpeername(Socket, (sockaddr*)&BSDAddr.Addr, &Size);
 
 	if (Result != 0)
 	{
@@ -339,34 +344,64 @@ bool FSocketBSD::SetBroadcast(bool bAllowBroadcast)
 
 bool FSocketBSD::JoinMulticastGroup(const FInternetAddr& GroupAddress)
 {
-	ip_mreq imr;
+	const FInternetAddrBSD& BSDAddr = static_cast<const FInternetAddrBSD&>(GroupAddress);
 
-	imr.imr_interface.s_addr = INADDR_ANY;
-	imr.imr_multiaddr = ((sockaddr_in*)&**((FInternetAddrBSD*)(&GroupAddress)))->sin_addr;
+#if PLATFORM_HAS_BSD_IPV6_SOCKETS
+	if (BSDAddr.GetProtocolFamily() == ESocketProtocolFamily::IPv6)
+	{
+		ipv6_mreq imr;
+		imr.ipv6mr_interface = 0;
+		imr.ipv6mr_multiaddr = ((sockaddr_in6*)&(BSDAddr.Addr))->sin6_addr;
+		return (setsockopt(Socket, IPPROTO_IPV6, IP_ADD_MEMBERSHIP, (char*)&imr, sizeof(imr)) == 0);
+	}
+#endif
 
+	ip_mreq imr;
+	imr.imr_interface.s_addr = INADDR_ANY;
+	imr.imr_multiaddr = ((sockaddr_in*)&(BSDAddr.Addr))->sin_addr;
 	return (setsockopt(Socket, IPPROTO_IP, IP_ADD_MEMBERSHIP, (char*)&imr, sizeof(imr)) == 0);
 }
 
 
 bool FSocketBSD::LeaveMulticastGroup(const FInternetAddr& GroupAddress)
 {
-	ip_mreq imr;
+	const FInternetAddrBSD& BSDAddr = static_cast<const FInternetAddrBSD&>(GroupAddress);
 
-	imr.imr_interface.s_addr = INADDR_ANY;
-	imr.imr_multiaddr = ((sockaddr_in*)&**((FInternetAddrBSD*)(&GroupAddress)))->sin_addr;
+#if PLATFORM_HAS_BSD_IPV6_SOCKETS
+	if (BSDAddr.GetProtocolFamily() == ESocketProtocolFamily::IPv6)
+	{
+		ipv6_mreq imr;
+		imr.ipv6mr_interface = 0;
+		imr.ipv6mr_multiaddr = ((sockaddr_in6*)&(BSDAddr.Addr))->sin6_addr;
+		return (setsockopt(Socket, IPPROTO_IPV6, IP_DROP_MEMBERSHIP, (char*)&imr, sizeof(imr)) == 0);
+	}
+#endif
 
+	ip_mreq imr;
+	imr.imr_interface.s_addr = INADDR_ANY;
+	imr.imr_multiaddr = ((sockaddr_in*)&(BSDAddr.Addr))->sin_addr;
 	return (setsockopt(Socket, IPPROTO_IP, IP_DROP_MEMBERSHIP, (char*)&imr, sizeof(imr)) == 0);
 }
 
 
 bool FSocketBSD::SetMulticastLoopback(bool bLoopback)
 {
+#if PLATFORM_HAS_BSD_IPV6_SOCKETS
+	bool IPv6Loop = bLoopback;
+	setsockopt(Socket, IPPROTO_IPV6, IPV6_MULTICAST_LOOP, (char*)&IPv6Loop, sizeof(IPv6Loop));
+#endif
+
 	return (setsockopt(Socket, IPPROTO_IP, IP_MULTICAST_LOOP, (char*)&bLoopback, sizeof(bLoopback)) == 0);
 }
 
 
 bool FSocketBSD::SetMulticastTtl(uint8 TimeToLive)
 {
+#if PLATFORM_HAS_BSD_IPV6_SOCKETS
+	uint8 IPv6TTL = TimeToLive;
+	setsockopt(Socket, IPPROTO_IPV6, IP_MULTICAST_TTL, (char*)&IPv6TTL, sizeof(IPv6TTL));
+#endif
+
 	return (setsockopt(Socket, IPPROTO_IP, IP_MULTICAST_TTL, (char*)&TimeToLive, sizeof(TimeToLive)) == 0);
 }
 
@@ -422,9 +457,8 @@ bool FSocketBSD::SetReceiveBufferSize(int32 Size,int32& NewSize)
 
 int32 FSocketBSD::GetPortNo(void)
 {
-	sockaddr_in Addr;
-
-	SOCKLEN Size = sizeof(sockaddr_in);
+	sockaddr_storage Addr;
+	SOCKLEN Size = sizeof(sockaddr_storage);
 
 	// Figure out what ip/port we are bound to
 	bool bOk = getsockname(Socket, (sockaddr*)&Addr, &Size) == 0;
@@ -437,8 +471,33 @@ int32 FSocketBSD::GetPortNo(void)
 		return 0;
 	}
 
+#if PLATFORM_HAS_BSD_IPV6_SOCKETS
+	if (Addr.ss_family == AF_INET6)
+	{
+		return ntohs(((sockaddr_in6&)Addr).sin6_port);
+	}
+#endif
+
 	// Convert big endian port to native endian port.
-	return ntohs(Addr.sin_port);
+	return ntohs(((sockaddr_in&)Addr).sin_port);
+}
+
+bool FSocketBSD::SetIPv6Only(bool bIPv6Only)
+{
+#if PLATFORM_HAS_BSD_IPV6_SOCKETS
+	int v6only = bIPv6Only ? 1 : 0;
+	bool bOk = setsockopt(Socket, IPPROTO_IPV6, IPV6_V6ONLY, (char*)&v6only, sizeof(v6only)) == 0;
+
+	if (bOk == false)
+	{
+		check(SocketSubsystem);
+		UE_LOG(LogSockets, Error, TEXT("Failed to set sock opt for socket (%s)"), SocketSubsystem->GetSocketError());
+	}
+
+	return bOk;
+#else 
+	return false;
+#endif
 }
 
 
diff --git a/Engine/Source/Runtime/Sockets/Private/BSDSockets/SocketsBSD.h b/Engine/Source/Runtime/Sockets/Private/BSDSockets/SocketsBSD.h
--- a/Engine/Source/Runtime/Sockets/Private/BSDSockets/SocketsBSD.h
+++ b/Engine/Source/Runtime/Sockets/Private/BSDSockets/SocketsBSD.h
@@ -7,7 +7,7 @@
 
 class FInternetAddr;
 
-#if PLATFORM_HAS_BSD_SOCKETS
+#if PLATFORM_HAS_BSD_SOCKETS || PLATFORM_HAS_BSD_IPV6_SOCKETS
 
 #include "Sockets.h"
 
@@ -15,7 +15,7 @@ class FInternetAddr;
 /**
  * Enumerates BSD socket state parameters.
  */
-enum class ESocketBSDParam
+enum class ESocketBSDParam : uint8
 {
 	CanRead,
 	CanWrite,
@@ -26,7 +26,7 @@ enum class ESocketBSDParam
 /**
  * Enumerates BSD socket state return values.
  */
-enum class ESocketBSDReturn
+enum class ESocketBSDReturn : uint8
 {
 	Yes,
 	No,
@@ -111,6 +111,8 @@ public:
 	virtual bool SetReceiveBufferSize(int32 Size,int32& NewSize) override;
 	virtual int32 GetPortNo() override;
 
+	bool SetIPv6Only(bool bIPv6Only);
+
 protected:
 
 	/** This is generally select(), but makes it easier for platforms without select to replace it. */
diff --git a/Engine/Source/Runtime/Sockets/Private/IPAddress.cpp b/Engine/Source/Runtime/Sockets/Private/IPAddress.cpp
--- a/Engine/Source/Runtime/Sockets/Private/IPAddress.cpp
+++ b/Engine/Source/Runtime/Sockets/Private/IPAddress.cpp
@@ -10,10 +10,7 @@
  */
 FResolveInfoCached::FResolveInfoCached(const FInternetAddr& InAddr)
 {
-	// @todo sockets: This should use Clone
-	uint32 IpAddr;
-	InAddr.GetIp(IpAddr);
-	Addr = ISocketSubsystem::Get()->CreateInternetAddr(IpAddr, InAddr.GetPort());
+	Addr = InAddr.Clone();
 }
 
 /**
diff --git a/Engine/Source/Runtime/Sockets/Private/SocketSubsystem.cpp b/Engine/Source/Runtime/Sockets/Private/SocketSubsystem.cpp
--- a/Engine/Source/Runtime/Sockets/Private/SocketSubsystem.cpp
+++ b/Engine/Source/Runtime/Sockets/Private/SocketSubsystem.cpp
@@ -293,50 +293,61 @@ TSharedRef<FInternetAddr> ISocketSubsystem::GetLocalHostAddr(FOutputDevice& Out,
 	HostAddr->SetAnyAddress();
 
 	bCanBindAll = false;
-
-	TCHAR Home[256]=TEXT("");
 	FString HostName;
+
+	if (GetMultihomeAddress(HostAddr))
+	{
+		return HostAddr;
+	}
+
 	if (GetHostName(HostName) == false)
 	{
 		Out.Logf(TEXT("%s: gethostname failed (%s)"), GetSocketAPIName(), GetSocketError());
 	}
-	if (FParse::Value(FCommandLine::Get(),TEXT("MULTIHOME="),Home,ARRAY_COUNT(Home)))
+
+	// Failing to find the host is not considered an error and we just bind to any address
+	ESocketErrors FindHostResult = GetHostByName(TCHAR_TO_ANSI(*HostName), *HostAddr);
+	if (FindHostResult == SE_NO_ERROR || FindHostResult == SE_HOST_NOT_FOUND 
+		|| FindHostResult == SE_EWOULDBLOCK || FindHostResult == SE_TRY_AGAIN)
 	{
-		bool bIsValid = false;
-		HostAddr->SetIp(Home,bIsValid);
-		if (Home == NULL || bIsValid == false)
+		if( !FParse::Param(FCommandLine::Get(),TEXT("PRIMARYNET")) )
 		{
-			Out.Logf( TEXT("Invalid multihome IP address %s"), Home );
+			bCanBindAll = true;
+		}
+		static bool First;
+		if( !First )
+		{
+			First = true;
+			UE_LOG(LogInit, Log, TEXT("%s: I am %s (%s)"), GetSocketAPIName(), *HostName, *HostAddr->ToString(true) );
 		}
 	}
 	else
 	{
-		// Failing to find the host is not considered an error and we just bind to any address
-		ESocketErrors FindHostResult = GetHostByName(TCHAR_TO_ANSI(*HostName), *HostAddr);
-		if (FindHostResult == SE_NO_ERROR || FindHostResult == SE_HOST_NOT_FOUND 
-			|| FindHostResult == SE_EWOULDBLOCK || FindHostResult == SE_TRY_AGAIN)
-		{
-			if( !FParse::Param(FCommandLine::Get(),TEXT("PRIMARYNET")) )
-			{
-				bCanBindAll = true;
-			}
-			static bool First;
-			if( !First )
-			{
-				First = true;
-				UE_LOG(LogInit, Log, TEXT("%s: I am %s (%s)"), GetSocketAPIName(), *HostName, *HostAddr->ToString(true) );
-			}
-		}
-		else
-		{
-			Out.Logf(TEXT("GetHostByName failed (%s)"), GetSocketError(FindHostResult));
-		}
+		Out.Logf(TEXT("GetHostByName failed (%s)"), GetSocketError(FindHostResult));
 	}
 
 	// return the newly created address
 	return HostAddr;
 }
 
+bool ISocketSubsystem::GetMultihomeAddress(TSharedRef<class FInternetAddr>& Addr)
+{
+	TCHAR Home[256] = TEXT("");
+	if (FParse::Value(FCommandLine::Get(), TEXT("MULTIHOME="), Home, ARRAY_COUNT(Home)))
+	{
+		bool bIsValid = false;
+		Addr->SetIp(Home, bIsValid);
+		if (Home == NULL || !bIsValid)
+		{
+			UE_LOG(LogSockets, Log, TEXT("Invalid multihome IP address %s"), Home);
+			return false;
+		}
+
+		return true;
+	}
+	return false;
+}
+
 bool ISocketSubsystem::GetHostByNameFromCache(const ANSICHAR* HostName, TSharedPtr<FInternetAddr>& Addr)
 {
 	// Lock for thread safety
diff --git a/Engine/Source/Runtime/Sockets/Private/Unix/SocketSubsystemUnix.cpp b/Engine/Source/Runtime/Sockets/Private/Unix/SocketSubsystemUnix.cpp
--- a/Engine/Source/Runtime/Sockets/Private/Unix/SocketSubsystemUnix.cpp
+++ b/Engine/Source/Runtime/Sockets/Private/Unix/SocketSubsystemUnix.cpp
@@ -4,6 +4,7 @@
 #include "Misc/CommandLine.h"
 #include "SocketSubsystemModule.h"
 #include "IPAddress.h"
+#include "BSDSockets/IPAddressBSD.h"
 
 #include <net/if.h>
 
@@ -102,55 +103,63 @@ TSharedRef<FInternetAddr> FSocketSubsystemUnix::GetLocalHostAddr(FOutputDevice&
 		return Addr;
 	}
 
-	// If superclass got the address from command line, honor that override
-	TCHAR Home[256]=TEXT("");
-	if (FParse::Value(FCommandLine::Get(),TEXT("MULTIHOME="),Home,ARRAY_COUNT(Home)))
+	// Grab multihome if it exists.
+	if (GetMultihomeAddress(Addr))
 	{
-		TSharedRef<FInternetAddr> TempAddr = CreateInternetAddr();
-		bool bIsValid = false;
-		TempAddr->SetIp(Home, bIsValid);
-		if (bIsValid)
-		{
-			return TempAddr;
-		}
+		return Addr;
+	}
+
+	TArray<TSharedPtr<FInternetAddr> > ResultArray;
+	if (GetLocalAdapterAddresses(ResultArray))
+	{
+		return ResultArray[0]->Clone();
+	}
+
+	// Fall back to this.
+	Addr->SetAnyAddress();
+	return Addr;
+}
+
+bool FSocketSubsystemUnix::GetLocalAdapterAddresses(TArray<TSharedPtr<FInternetAddr> >& OutAddresses)
+{
+	TSharedRef<FInternetAddr> MultihomeAddress = FSocketSubsystemBSD::CreateInternetAddr();
+	if (GetMultihomeAddress(MultihomeAddress))
+	{
+		OutAddresses.Add(MultihomeAddress);
 	}
 
-	// we need to go deeper...  (see http://unixhelp.ed.ac.uk/CGI/man-cgi?netdevice+7)
-	int TempSocket = socket(PF_INET, SOCK_STREAM, 0);
-	if (TempSocket)
+	ifaddrs* Interfaces = NULL;
+	int InterfaceQueryRet = getifaddrs(&Interfaces);
+	if (InterfaceQueryRet == 0)
 	{
-		ifreq IfReqs[8];
-		
-		ifconf IfConfig;
-		FMemory::Memzero(IfConfig);
-		IfConfig.ifc_req = IfReqs;
-		IfConfig.ifc_len = sizeof(IfReqs);
-		
-		int Result = ioctl(TempSocket, SIOCGIFCONF, &IfConfig);
-		if (Result == 0)
+		// Loop through linked list of interfaces
+		for (ifaddrs* Travel = Interfaces; Travel != NULL; Travel = Travel->ifa_next)
 		{
-			for (int32 IdxReq = 0; IdxReq < ARRAY_COUNT(IfReqs); ++IdxReq)
+			// Skip over empty data sets.
+			if (Travel->ifa_addr == NULL)
 			{
-				// grab the first non-loobpack one which is up
-				int ResultFlags = ioctl(TempSocket, SIOCGIFFLAGS, &IfReqs[IdxReq]);
-				if (ResultFlags == 0 && 
-					(IfReqs[IdxReq].ifr_flags & IFF_UP) && 
-					(IfReqs[IdxReq].ifr_flags & IFF_LOOPBACK) == 0)
-				{
-					int32 NetworkAddr = reinterpret_cast<sockaddr_in *>(&IfReqs[IdxReq].ifr_addr)->sin_addr.s_addr;
-					Addr->SetIp(ntohl(NetworkAddr));
-					break;
-				}
+				continue;
+			}
+
+			int& family = Travel->ifa_addr->sa_family;
+			// Find any up and non-loopback addresses
+			if ((Travel->ifa_flags & IFF_UP) && 
+				(Travel->ifa_flags & IFF_LOOPBACK) == 0 && 
+				(family == AF_INET || family == AF_INET6))
+			{
+				TSharedRef<FInternetAddrBSD> NewAddress = MakeShareable(new FInternetAddrBSD);
+				NewAddress->SetIp(*((sockaddr_storage*)Travel->ifa_addr));
+				OutAddresses.Add(NewAddress);
 			}
 		}
-		else
-		{
-			int ErrNo = errno;
-			UE_LOG(LogSockets, Warning, TEXT("ioctl( ,SIOGCIFCONF, ) failed, errno=%d (%s)"), ErrNo, ANSI_TO_TCHAR(strerror(ErrNo)));
-		}
-		
-		close(TempSocket);
+
+		freeifaddrs(Interfaces);
+	}
+	else
+	{
+		UE_LOG(Log_Sockets, Warning, TEXT("getifaddrs returned result %d"), InterfaceQueryRet);
+		return (OutAddresses.Num() == 0); // if getifaddrs somehow doesn't work but we have multihome, then it's fine.
 	}
 
-	return Addr;
+	return true;
 }
diff --git a/Engine/Source/Runtime/Sockets/Private/Unix/SocketSubsystemUnix.h b/Engine/Source/Runtime/Sockets/Private/Unix/SocketSubsystemUnix.h
--- a/Engine/Source/Runtime/Sockets/Private/Unix/SocketSubsystemUnix.h
+++ b/Engine/Source/Runtime/Sockets/Private/Unix/SocketSubsystemUnix.h
@@ -51,4 +51,5 @@ public:
 	virtual void Shutdown() override;
 	virtual bool HasNetworkDevice() override;
 	virtual TSharedRef<FInternetAddr> GetLocalHostAddr(FOutputDevice& Out, bool& bCanBindAll) override;
+	virtual bool GetLocalAdapterAddresses(TArray<TSharedPtr<FInternetAddr> >& OutAddresses) override;
 };
diff --git a/Engine/Source/Runtime/Sockets/Private/Windows/SocketSubsystemWindows.cpp b/Engine/Source/Runtime/Sockets/Private/Windows/SocketSubsystemWindows.cpp
--- a/Engine/Source/Runtime/Sockets/Private/Windows/SocketSubsystemWindows.cpp
+++ b/Engine/Source/Runtime/Sockets/Private/Windows/SocketSubsystemWindows.cpp
@@ -3,6 +3,7 @@
 #include "SocketSubsystemWindows.h"
 #include "SocketSubsystemModule.h"
 #include "Modules/ModuleManager.h"
+#include "BSDSockets/IPAddressBSD.h"
 
 #include "Windows/AllowWindowsPlatformTypes.h"
 #include "Iphlpapi.h"
@@ -70,26 +71,7 @@ void FSocketSubsystemWindows::Destroy()
 
 FSocket* FSocketSubsystemWindows::CreateSocket(const FName& SocketType, const FString& SocketDescription, bool bForceUDP)
 {
-	SOCKET Socket = INVALID_SOCKET;
-	FSocketBSD* NewSocket = nullptr;
-
-	switch (SocketType.GetComparisonIndex())
-	{
-	case NAME_DGram:
-		// Creates a data gram (UDP) socket
-		Socket = socket(AF_INET, SOCK_DGRAM, IPPROTO_UDP);
-		NewSocket = (Socket != INVALID_SOCKET) ? InternalBSDSocketFactory(Socket, SOCKTYPE_Datagram, SocketDescription) : nullptr;
-		break;
-
-	case NAME_Stream:
-		// Creates a stream (TCP) socket
-		Socket = socket(AF_INET, SOCK_STREAM, IPPROTO_TCP);
-		NewSocket = (Socket != INVALID_SOCKET) ? InternalBSDSocketFactory(Socket, SOCKTYPE_Streaming, SocketDescription) : nullptr;
-		break;
-
-	default:
-		break;
-	}
+	FSocketBSD* NewSocket = (FSocketBSD*)FSocketSubsystemBSD::CreateSocket(SocketType, SocketDescription, bForceUDP);
 
 	if (NewSocket != nullptr)
 	{
@@ -183,9 +165,10 @@ bool FSocketSubsystemWindows::GetLocalAdapterAddresses( TArray<TSharedPtr<FInter
 			{
 				if ((UnicastAddress->Flags & IP_ADAPTER_ADDRESS_DNS_ELIGIBLE) != 0)
 				{
-					uint32 RawAddress = ntohl(*((uint32*)(&UnicastAddress->Address.lpSockaddr->sa_data[2])));
-
-					OutAdresses.Add(CreateInternetAddr(RawAddress));
+					const sockaddr_storage* RawAddress = (const sockaddr_storage*)(UnicastAddress->Address.lpSockaddr);
+					TSharedRef<FInternetAddr> NewAddress = MakeShareable(new FInternetAddrBSD);
+					static_cast<FInternetAddrBSD&>(NewAddress.Get()).SetIp(*RawAddress);
+					OutAdresses.Add(NewAddress);
 				}
 			}
 		}
diff --git a/Engine/Source/Runtime/Sockets/Public/IPAddress.h b/Engine/Source/Runtime/Sockets/Public/IPAddress.h
--- a/Engine/Source/Runtime/Sockets/Public/IPAddress.h
+++ b/Engine/Source/Runtime/Sockets/Public/IPAddress.h
@@ -56,7 +56,10 @@ public:
 	 *
 	 * @param OutPort the host byte order int that receives the port
 	 */
-	virtual void GetPort(int32& OutPort) const = 0;
+	virtual void GetPort(int32& OutPort) const
+	{
+		OutPort = GetPort();
+	}
 
 	/**
 	 * Returns the port number from this address in host byte order
@@ -79,12 +82,29 @@ public:
 		return GetPort();
 	}
 
+	/**
+	 * Sets the ip address from a raw network byte order array.
+	 *
+	 * @param RawAddr the new address to use (must be converted to network byte order)
+	 */
+	virtual void SetRawIp(const TArray<uint8>& RawAddr) = 0;
+
+	/**
+	 * Gets the ip address in a raw array stored in network byte order.
+	 *
+	 * @return The raw address stored in an array
+	 */
+	virtual TArray<uint8> GetRawIp() const = 0;
+
 	/** Sets the address to be any address */
 	virtual void SetAnyAddress() = 0;
 
 	/** Sets the address to broadcast */
 	virtual void SetBroadcastAddress() = 0;
 
+	/** Sets the address to loopback */
+	virtual void SetLoopbackAddress() = 0;
+
 	/**
 	 * Converts this internet ip address to string form
 	 *
@@ -99,9 +119,8 @@ public:
 	 */
 	virtual bool operator==(const FInternetAddr& Other) const
 	{
-		uint32 ThisIP, OtherIP;
-		GetIp(ThisIP);
-		Other.GetIp(OtherIP);
+		TArray<uint8> ThisIP = GetRawIp();
+		TArray<uint8> OtherIP = Other.GetRawIp();
 		return ThisIP == OtherIP && GetPort() == Other.GetPort();
 	}
 
@@ -119,25 +138,12 @@ public:
 	 */
 	virtual bool IsValid() const = 0;
 	
-#if PLATFORM_IOS
-	virtual void Copy(const FInternetAddr& Other) { check(false); }
-#endif
-
 	/**
-	 * Determines if the string is in IP address form or needs host resolution
-	 *
-	 * @param IpAddr the string to check for being a well formed ip
+	 * Creates a new structure with the same data as this structure
 	 *
-	 * @return true if the ip address was parseable, false otherwise
+	 * @return The new structure
 	 */
-// 	static bool IsValidIp(const TCHAR* IpAddr) override
-// 	{
-// 		FInternetAddr LocalAddr;
-// 		bool bIsValid = false;
-// 		LocalAddr.SetIp(IpAddr,bIsValid);
-// 		return bIsValid;
-// 	}
-
+	virtual TSharedRef<FInternetAddr> Clone() const = 0;
 };
 
 /**
diff --git a/Engine/Source/Runtime/Sockets/Public/SocketSubsystem.h b/Engine/Source/Runtime/Sockets/Public/SocketSubsystem.h
--- a/Engine/Source/Runtime/Sockets/Public/SocketSubsystem.h
+++ b/Engine/Source/Runtime/Sockets/Public/SocketSubsystem.h
@@ -86,12 +86,12 @@ public:
 	virtual class FSocket* CreateSocket(const FName& SocketType, const FString& SocketDescription, bool bForceUDP = false) = 0;
 
 	/**
-	* Creates a resolve info cached struct to hold the resolved address
-	*
-	* @Param Addr address to resolve for the socket subsystem
-	*
-	* @return the new resolved address or NULL if failed
-	*/
+	 * Creates a resolve info cached struct to hold the resolved address
+	 *
+	 * @Param Addr address to resolve for the socket subsystem
+	 *
+	 * @return the new resolved address or NULL if failed
+	 */
 	virtual class FResolveInfoCached* CreateResolveInfoCached(TSharedPtr<FInternetAddr> Addr) const;
 
 	/**
@@ -141,7 +141,7 @@ public:
 	virtual bool GetHostName(FString& HostName) = 0;
 
 	/**
-	 *	Create a proper FInternetAddr representation
+	 * Create a proper FInternetAddr representation
 	 * @param Address host address
 	 * @param Port host port
 	 */
@@ -214,6 +214,16 @@ public:
 	 */
 	virtual TSharedRef<FInternetAddr> GetLocalHostAddr(FOutputDevice& Out, bool& bCanBindAll);
 
+	/**
+	 * Returns the multihome address if the flag is present and valid. For ease of use, 
+	 * this function will check validity of the address for the caller.
+	 * 
+	 * @param Addr the address structure which will have the Multihome address in it if set.
+	 *
+	 * @return If the multihome address is valid or not.
+	 */
+	virtual bool GetMultihomeAddress(TSharedRef<class FInternetAddr>& Addr);
+
 	/**
 	 * Checks the host name cache for an existing entry (faster than resolving again)
 	 *
diff --git a/Engine/Source/Runtime/Sockets/Public/SocketTypes.h b/Engine/Source/Runtime/Sockets/Public/SocketTypes.h
--- a/Engine/Source/Runtime/Sockets/Public/SocketTypes.h
+++ b/Engine/Source/Runtime/Sockets/Public/SocketTypes.h
@@ -4,6 +4,16 @@
 
 #include "CoreMinimal.h"
 
+/** Indicates the socket protocol of socket being used, typically for BSD Sockets */
+enum class ESocketProtocolFamily : uint32
+{
+	/** No protocol family specification. Typically defined as AF_UNSPEC */
+	None,
+	/** IPv4 and IPv6 respectively. */
+	IPv4,
+	IPv6
+};
+
 /** Indicates the type of socket being used (streaming or datagram) */
 enum ESocketType
 {
