New IP Addr changes for iOS

diff --git a/Engine/Source/Runtime/Sockets/Private/IOS/IPAddressBSDIOS.h b/Engine/Source/Runtime/Sockets/Private/IOS/IPAddressBSDIOS.h
new file mode 100644
--- /dev/null
+++ b/Engine/Source/Runtime/Sockets/Private/IOS/IPAddressBSDIOS.h
@@ -0,0 +1,27 @@
+// Copyright 1998-2018 Epic Games, Inc. All Rights Reserved.
+
+#pragma once
+
+#include "CoreMinimal.h"
+#include "BSDSockets/SocketSubsystemBSDPrivate.h"
+#include "BSDSockets/IPAddressBSD.h"
+
+#if PLATFORM_HAS_BSD_IPV6_SOCKETS
+
+class FInternetAddrBSDIOS : public FInternetAddrBSD
+{
+
+public:
+
+	/** Sets the address to broadcast */
+	virtual void SetBroadcastAddress() override
+	{
+		UE_LOG(LogSockets, Verbose, TEXT("SetBroadcastAddress() FInternetAddrBSDIOS"));
+		SetIp(INADDR_BROADCAST);
+		SetPort(0);
+	}
+
+	// TODO: Determine if we need to override the others.
+};
+
+#endif
diff --git a/Engine/Source/Runtime/Sockets/Private/IOS/IPAddressBSDIPv6IOS.h b/Engine/Source/Runtime/Sockets/Private/IOS/IPAddressBSDIPv6IOS.h
deleted file mode 100644
diff --git a/Engine/Source/Runtime/Sockets/Private/IOS/SocketSubsystemIOS.cpp b/Engine/Source/Runtime/Sockets/Private/IOS/SocketSubsystemIOS.cpp
--- a/Engine/Source/Runtime/Sockets/Private/IOS/SocketSubsystemIOS.cpp
+++ b/Engine/Source/Runtime/Sockets/Private/IOS/SocketSubsystemIOS.cpp
@@ -10,10 +10,10 @@
 #include "IPAddressBSDIPv6IOS.h"
 FSocketSubsystemIOS* FSocketSubsystemIOS::SocketSingleton = NULL;
 
-class FSocketBSDIPv6* FSocketSubsystemIOS::InternalBSDSocketFactory(SOCKET Socket, ESocketType SocketType, const FString& SocketDescription)
+class FSocketBSD* FSocketSubsystemIOS::InternalBSDSocketFactory(SOCKET Socket, ESocketType SocketType, const FString& SocketDescription)
 {
 	UE_LOG(LogIOS, Log, TEXT(" FSocketSubsystemIOS::InternalBSDSocketFactory"));
-	return new FSocketBSDIPv6IOS(Socket, SocketType, SocketDescription, this);
+	return new FSocketBSDIOS(Socket, SocketType, SocketDescription, this);
 }
 
 FName CreateSocketSubsystem( FSocketSubsystemModule& SocketSubsystemModule )
@@ -77,7 +77,7 @@ bool FSocketSubsystemIOS::HasNetworkDevice()
 
 FSocket* FSocketSubsystemIOS::CreateSocket(const FName& SocketType, const FString& SocketDescription, bool bForceUDP)
 {
-	FSocketBSDIPv6* NewSocket = (FSocketBSDIPv6*)FSocketSubsystemBSDIPv6::CreateSocket(SocketType, SocketDescription, bForceUDP);
+	FSocketBSD* NewSocket = (FSocketBSD*)FSocketSubsystemBSD::CreateSocket(SocketType, SocketDescription, bForceUDP);
 	if (NewSocket)
 	{
 		NewSocket->SetIPv6Only(false);
@@ -91,7 +91,7 @@ FSocket* FSocketSubsystemIOS::CreateSocket(const FName& SocketType, const FStrin
 
 TSharedRef<FInternetAddr> FSocketSubsystemIOS::GetLocalHostAddr(FOutputDevice& Out, bool& bCanBindAll)
 {
-	TSharedRef<FInternetAddrBSDIPv6IOS> HostAddr = MakeShareable(new FInternetAddrBSDIPv6IOS);
+	TSharedRef<FInternetAddrBSDIOS> HostAddr = MakeShareable(new FInternetAddrBSDIOS);
 	HostAddr->SetAnyAddress();
 
 	ifaddrs* Interfaces = NULL;
@@ -105,17 +105,23 @@ TSharedRef<FInternetAddr> FSocketSubsystemIOS::GetLocalHostAddr(FOutputDevice& O
 		// Loop through linked list of interfaces
 		for (ifaddrs* Travel = Interfaces; Travel != NULL; Travel = Travel->ifa_next)
 		{
+			if (Travel->ifa_addr == NULL)
+			{
+				continue;
+			}
+
+			sockaddr_storage* AddrData = reinterpret_cast<sockaddr_storage*>(Travel->ifa_addr);
 			if (Travel->ifa_addr->sa_family == AF_INET6)
 			{
 				if (strcmp(Travel->ifa_name, "en0") == 0)
 				{
-					HostAddr->SetIp(*((sockaddr_in6*)Travel->ifa_addr));
+					HostAddr->SetIp(*AddrData);
 					bWasWifiSet = true;
 					bWasIPv6Set = true;
 				}
 				else if (!bWasWifiSet && strcmp(Travel->ifa_name, "pdp_ip0") == 0)
 				{
-					HostAddr->SetIp(*((sockaddr_in6*)Travel->ifa_addr));
+					HostAddr->SetIp(*AddrData);
 					bWasCellSet = true;
 				}
 			}
@@ -123,12 +129,12 @@ TSharedRef<FInternetAddr> FSocketSubsystemIOS::GetLocalHostAddr(FOutputDevice& O
 			{
 				if (strcmp(Travel->ifa_name, "en0") == 0)
 				{
-					HostAddr->SetIp(ntohl(((sockaddr_in*)Travel->ifa_addr)->sin_addr.s_addr));
+					HostAddr->SetIp(*AddrData);
 					bWasWifiSet = true;
 				}
 				else if (!bWasWifiSet && strcmp(Travel->ifa_name, "pdp_ip0") == 0)
 				{
-					HostAddr->SetIp(((sockaddr_in*)Travel->ifa_addr)->sin_addr.s_addr);
+					HostAddr->SetIp(*AddrData);
 					bWasCellSet = true;
 				}
 			}
@@ -157,8 +163,8 @@ TSharedRef<FInternetAddr> FSocketSubsystemIOS::GetLocalHostAddr(FOutputDevice& O
 }
 
 TSharedRef<FInternetAddr> FSocketSubsystemIOS::CreateInternetAddr(uint32 Address, uint32 Port)
-{
-	TSharedRef<FInternetAddr> Result = MakeShareable(new FInternetAddrBSDIPv6IOS);
+	{
+	TSharedRef<FInternetAddr> Result = MakeShareable(new FInternetAddrBSDIOS);
 	Result->SetIp(Address);
 	Result->SetPort(Port);
 	return Result;
diff --git a/Engine/Source/Runtime/Sockets/Private/IOS/SocketSubsystemIOS.h b/Engine/Source/Runtime/Sockets/Private/IOS/SocketSubsystemIOS.h
--- a/Engine/Source/Runtime/Sockets/Private/IOS/SocketSubsystemIOS.h
+++ b/Engine/Source/Runtime/Sockets/Private/IOS/SocketSubsystemIOS.h
@@ -3,13 +3,13 @@
 #pragma once
 
 #include "SocketSubsystem.h"
-#include "BSDIPv6Sockets/SocketSubsystemBSDIPv6.h"
+#include "BSDSockets/SocketSubsystemBSD.h"
 #include "SocketSubsystemPackage.h"
 
 /**
- * Windows specific socket subsystem implementation
+ * iOS specific socket subsystem implementation
  */
-class FSocketSubsystemIOS : public FSocketSubsystemBSDIPv6
+class FSocketSubsystemIOS : public FSocketSubsystemBSD
 {
 protected:
 	virtual TSharedRef<FInternetAddr> CreateInternetAddr(uint32 Address = 0, uint32 Port = 0) override;
@@ -30,10 +30,16 @@ public:
 	static FSocketSubsystemIOS* Create();
 
 	/**
-	 * Performs Windows specific socket clean up
+	 * Performs iOS specific socket clean up
 	 */
 	static void Destroy();
 
+
+	virtual ESocketProtocolFamily GetDefaultSocketProtocolFamily() const override
+	{
+		return ESocketProtocolFamily::IPv6;
+	}
+
 public:
 
 	FSocketSubsystemIOS() 
@@ -51,5 +57,5 @@ public:
 	virtual FSocket* CreateSocket(const FName& SocketType, const FString& SocketDescription, bool bForceUDP) override;
 
 	virtual TSharedRef<FInternetAddr> GetLocalHostAddr(FOutputDevice& Out, bool& bCanBindAll) override;
-	virtual class FSocketBSDIPv6* InternalBSDSocketFactory(SOCKET Socket, ESocketType SocketType, const FString& SocketDescription) override;
+	virtual class FSocketBSD* InternalBSDSocketFactory(SOCKET Socket, ESocketType SocketType, const FString& SocketDescription) override;
 };
\ No newline at end of file
diff --git a/Engine/Source/Runtime/Sockets/Private/IOS/SocketsBSDIOS.h b/Engine/Source/Runtime/Sockets/Private/IOS/SocketsBSDIOS.h
new file mode 100644
--- /dev/null
+++ b/Engine/Source/Runtime/Sockets/Private/IOS/SocketsBSDIOS.h
@@ -0,0 +1,35 @@
+// Copyright 1998-2018 Epic Games, Inc. All Rights Reserved.
+
+#pragma once
+
+#include "CoreMinimal.h"
+#include "BSDSockets/SocketSubsystemBSDPrivate.h"
+#include "Sockets.h"
+#include "BSDSockets/SocketsBSD.h"
+
+class FInternetAddr;
+
+#if PLATFORM_HAS_BSD_IPV6_SOCKETS
+
+
+/**
+ * Implements a BSD network socket on IOS.
+ */
+class FSocketBSDIOS
+	: public FSocketBSD
+{
+public:
+
+	FSocketBSDIOS(SOCKET InSocket, ESocketType InSocketType, const FString& InSocketDescription, ISocketSubsystem* InSubsystem)
+		:FSocketBSD(InSocket, InSocketType, InSocketDescription, InSubsystem)
+	{
+	}
+
+	virtual ~FSocketBSDIOS()
+	{	
+		FSocketBSD::Close();
+	}
+
+};
+
+#endif
diff --git a/Engine/Source/Runtime/Sockets/Private/IOS/SocketsBSDIPv6IOS.h b/Engine/Source/Runtime/Sockets/Private/IOS/SocketsBSDIPv6IOS.h
deleted file mode 100644
