Fixing Unix SocketSubsystem compilation issues.

diff --git a/Engine/Source/Runtime/Sockets/Private/Unix/SocketSubsystemUnix.cpp b/Engine/Source/Runtime/Sockets/Private/Unix/SocketSubsystemUnix.cpp
--- a/Engine/Source/Runtime/Sockets/Private/Unix/SocketSubsystemUnix.cpp
+++ b/Engine/Source/Runtime/Sockets/Private/Unix/SocketSubsystemUnix.cpp
@@ -6,6 +6,7 @@
 #include "IPAddress.h"
 #include "BSDSockets/IPAddressBSD.h"
 
+#include <ifaddrs.h>
 #include <net/if.h>
 
 FSocketSubsystemUnix* FSocketSubsystemUnix::SocketSingleton = NULL;
@@ -157,7 +158,7 @@ bool FSocketSubsystemUnix::GetLocalAdapterAddresses(TArray<TSharedPtr<FInternetA
 	}
 	else
 	{
-		UE_LOG(Log_Sockets, Warning, TEXT("getifaddrs returned result %d"), InterfaceQueryRet);
+		UE_LOG(LogSockets, Warning, TEXT("getifaddrs returned result %d"), InterfaceQueryRet);
 		return (OutAddresses.Num() == 0); // if getifaddrs somehow doesn't work but we have multihome, then it's fine.
 	}
 
