More Unix compilation fixes

diff --git a/Engine/Source/Runtime/Sockets/Private/Unix/SocketSubsystemUnix.cpp b/Engine/Source/Runtime/Sockets/Private/Unix/SocketSubsystemUnix.cpp
--- a/Engine/Source/Runtime/Sockets/Private/Unix/SocketSubsystemUnix.cpp
+++ b/Engine/Source/Runtime/Sockets/Private/Unix/SocketSubsystemUnix.cpp
@@ -142,11 +142,11 @@ bool FSocketSubsystemUnix::GetLocalAdapterAddresses(TArray<TSharedPtr<FInternetA
 				continue;
 			}
 
-			int& family = Travel->ifa_addr->sa_family;
+			uint16 AddrFamily = Travel->ifa_addr->sa_family;
 			// Find any up and non-loopback addresses
 			if ((Travel->ifa_flags & IFF_UP) && 
 				(Travel->ifa_flags & IFF_LOOPBACK) == 0 && 
-				(family == AF_INET || family == AF_INET6))
+				(AddrFamily == AF_INET || AddrFamily == AF_INET6))
 			{
 				TSharedRef<FInternetAddrBSD> NewAddress = MakeShareable(new FInternetAddrBSD);
 				NewAddress->SetIp(*((sockaddr_storage*)Travel->ifa_addr));
