Fixes for compilation and static analysis

diff --git a/Engine/Plugins/Online/OnlineSubsystemOculus/Source/Private/IPAddressOculus.h b/Engine/Plugins/Online/OnlineSubsystemOculus/Source/Private/IPAddressOculus.h
--- a/Engine/Plugins/Online/OnlineSubsystemOculus/Source/Private/IPAddressOculus.h
+++ b/Engine/Plugins/Online/OnlineSubsystemOculus/Source/Private/IPAddressOculus.h
@@ -80,7 +80,7 @@ public:
 
 	virtual void SetRawIp(const TArray<uint8>& RawAddr) override
 	{
-		ovrID NewId;
+		ovrID NewId = 0;
 
 		// Make a quick copy of the array
 		TArray<uint8> WorkingArray = RawAddr;
diff --git a/Engine/Source/Runtime/Sockets/Private/Android/SocketSubsystemAndroid.cpp b/Engine/Source/Runtime/Sockets/Private/Android/SocketSubsystemAndroid.cpp
--- a/Engine/Source/Runtime/Sockets/Private/Android/SocketSubsystemAndroid.cpp
+++ b/Engine/Source/Runtime/Sockets/Private/Android/SocketSubsystemAndroid.cpp
@@ -163,18 +163,18 @@ TSharedRef<FInternetAddr> FSocketSubsystemAndroid::GetLocalHostAddr(FOutputDevic
 					if (strcmp(IfReqs[IdxReq].ifr_name, "wlan0") == 0)
 					{
 						// 'Usually' wifi, Prefer wifi
-						FMemory::Memcpy((void*)&WifiAddress, (void*)IfReqs[IdxReq].ifr_addr, sizeof(sockaddr_in));
+						FMemory::Memcpy((void*)&WifiAddress, (void*)(&IfReqs[IdxReq].ifr_addr), sizeof(sockaddr_in));
 						break;
 					}
 					else if (strcmp(IfReqs[IdxReq].ifr_name, "rmnet0") == 0)
 					{
 						// 'Usually' cellular
-						FMemory::Memcpy((void*)&CellularAddress, (void*)IfReqs[IdxReq].ifr_addr, sizeof(sockaddr_in));
+						FMemory::Memcpy((void*)&CellularAddress, (void*)(&IfReqs[IdxReq].ifr_addr), sizeof(sockaddr_in));
 					}
 					else if (OtherAddress.ss_family == AF_UNSPEC)
 					{
 						// First alternate found
-						FMemory::Memcpy((void*)&OtherAddress, (void*)IfReqs[IdxReq].ifr_addr, sizeof(sockaddr_in));
+						FMemory::Memcpy((void*)&OtherAddress, (void*)(&IfReqs[IdxReq].ifr_addr), sizeof(sockaddr_in));
 					}
 				}
 			}
diff --git a/Engine/Source/Runtime/Sockets/Private/IOS/SocketSubsystemIOS.cpp b/Engine/Source/Runtime/Sockets/Private/IOS/SocketSubsystemIOS.cpp
--- a/Engine/Source/Runtime/Sockets/Private/IOS/SocketSubsystemIOS.cpp
+++ b/Engine/Source/Runtime/Sockets/Private/IOS/SocketSubsystemIOS.cpp
@@ -6,8 +6,9 @@
 #include "BSDSockets/SocketsBSD.h"
 #include "IPAddress.h"
 #include <ifaddrs.h>
-#include "SocketsBSDIPv6IOS.h"
-#include "IPAddressBSDIPv6IOS.h"
+#include "SocketsBSDIOS.h"
+#include "IPAddressBSDIOS.h"
+
 FSocketSubsystemIOS* FSocketSubsystemIOS::SocketSingleton = NULL;
 
 class FSocketBSD* FSocketSubsystemIOS::InternalBSDSocketFactory(SOCKET Socket, ESocketType SocketType, const FString& SocketDescription)
