Minor Address Fixes

diff --git a/Engine/Source/Runtime/Sockets/Private/Windows/SocketSubsystemWindows.cpp b/Engine/Source/Runtime/Sockets/Private/Windows/SocketSubsystemWindows.cpp
--- a/Engine/Source/Runtime/Sockets/Private/Windows/SocketSubsystemWindows.cpp
+++ b/Engine/Source/Runtime/Sockets/Private/Windows/SocketSubsystemWindows.cpp
@@ -166,8 +166,8 @@ bool FSocketSubsystemWindows::GetLocalAdapterAddresses( TArray<TSharedPtr<FInter
 				if ((UnicastAddress->Flags & IP_ADAPTER_ADDRESS_DNS_ELIGIBLE) != 0)
 				{
 					const sockaddr_storage* RawAddress = (const sockaddr_storage*)(UnicastAddress->Address.lpSockaddr);
-					TSharedRef<FInternetAddr> NewAddress = MakeShareable(new FInternetAddrBSD);
-					static_cast<FInternetAddrBSD&>(NewAddress.Get()).SetIp(*RawAddress);
+					TSharedRef<FInternetAddrBSD> NewAddress = MakeShareable(new FInternetAddrBSD);
+					NewAddress->SetIp(*RawAddress);
 					OutAdresses.Add(NewAddress);
 				}
 			}
