Fix the deprecation warnings by removing the macro from the Sockets and IPAddress classes. While these classes are deprecated, can't mark them as such due to how they are generated and used.

On some platforms that are more picky about these deprecations, they return incorrect warnings that otherwise do not show up on other platforms.

In addition, expose and add a function for getaddrinfo for IPv6

diff --git a/Engine/Source/Runtime/Sockets/Private/BSDIPv6Sockets/IPAddressBSDIPv6.h b/Engine/Source/Runtime/Sockets/Private/BSDIPv6Sockets/IPAddressBSDIPv6.h
--- a/Engine/Source/Runtime/Sockets/Private/BSDIPv6Sockets/IPAddressBSDIPv6.h
+++ b/Engine/Source/Runtime/Sockets/Private/BSDIPv6Sockets/IPAddressBSDIPv6.h
@@ -9,21 +9,17 @@
 
 #if PLATFORM_HAS_BSD_IPV6_SOCKETS
 
-class FInternetAddrBSDIPv6_Base_DEPRECATED
+/**
+* Represents an internet ip address, using the relatively standard sockaddr_in6 structure. All data is in network byte order
+*/
+class FInternetAddrBSDIPv6 : public FInternetAddr
 {
-protected:
+private:
 	/** The internet ip address structure */
 	sockaddr_in6 Addr;
 
 	/** Horrible hack to catch hard coded multicasting on IPv4 **/
 	static const uint32 IPv4MulitcastAddr = ((230 << 24) | (0 << 16) | (0 << 8) | (1 << 0));
-};
-
-/**
-* Represents an internet ip address, using the relatively standard sockaddr_in6 structure. All data is in network byte order
-*/
-class DEPRECATED(4.21, "FInternetAddrBSDIPv6 is deprecated. Move to FInternetAddrBSD.") FInternetAddrBSDIPv6 : public FInternetAddr, FInternetAddrBSDIPv6_Base_DEPRECATED
-{
 
 public:
 	/**
diff --git a/Engine/Source/Runtime/Sockets/Private/BSDIPv6Sockets/SocketSubsystemBSDIPv6.cpp b/Engine/Source/Runtime/Sockets/Private/BSDIPv6Sockets/SocketSubsystemBSDIPv6.cpp
--- a/Engine/Source/Runtime/Sockets/Private/BSDIPv6Sockets/SocketSubsystemBSDIPv6.cpp
+++ b/Engine/Source/Runtime/Sockets/Private/BSDIPv6Sockets/SocketSubsystemBSDIPv6.cpp
@@ -82,51 +82,73 @@ void FSocketSubsystemBSDIPv6::DestroySocket(FSocket* Socket)
 	delete Socket;
 }
 
-ESocketErrors FSocketSubsystemBSDIPv6::GetHostByName(const ANSICHAR* HostName, FInternetAddr& OutAddr)
+TArray<TSharedRef<FInternetAddr> > FSocketSubsystemBSDIPv6::GetAddressInfo(const ANSICHAR* HostName, bool bResolveAddress, ESocketProtocolFamily ProtocolType)
 {
-	FScopeLock ScopeLock(&HostByNameSynch);
-	addrinfo* AddrInfo = NULL;
+	TArray<TSharedRef<FInternetAddr> > GetAddrArray;
+
+#if PLATFORM_HAS_BSD_SOCKET_FEATURE_GETADDRINFO
+	addrinfo* AddrInfo = nullptr;
 
-	// We are only interested in IPv6 addresses.
 	addrinfo HintAddrInfo;
 	FMemory::Memzero(&HintAddrInfo, sizeof(HintAddrInfo));
 	HintAddrInfo.ai_family = AF_UNSPEC;
+	HintAddrInfo.ai_flags = AI_ADDRCONFIG | AI_PASSIVE;
 
-	int32 ErrorCode = getaddrinfo(HostName, NULL, &HintAddrInfo, &AddrInfo);
+	if (!bResolveAddress)
+	{
+		HintAddrInfo.ai_flags |= AI_NUMERICHOST;
+	}
+
+	int32 ErrorCode = getaddrinfo(HostName, nullptr, &HintAddrInfo, &AddrInfo);
 	ESocketErrors SocketError = TranslateGAIErrorCode(ErrorCode);
 	if (SocketError == SE_NO_ERROR)
 	{
 		for (; AddrInfo != nullptr; AddrInfo = AddrInfo->ai_next)
 		{
-			if (AddrInfo->ai_family == AF_INET6)
+			if (AddrInfo->ai_family == AF_INET6 || AddrInfo->ai_family == AF_INET)
 			{
-				sockaddr_in6* IPv6SockAddr = reinterpret_cast<sockaddr_in6*>(AddrInfo->ai_addr);
-				if (IPv6SockAddr != nullptr)
+				TSharedRef<FInternetAddrBSDIPv6> NewAddr = MakeShareable(new FInternetAddrBSDIPv6);
+				if (AddrInfo->ai_family == AF_INET6)
 				{
+					sockaddr_in6* IPv6SockAddr = reinterpret_cast<sockaddr_in6*>(AddrInfo->ai_addr);
+					if (IPv6SockAddr != nullptr)
+					{
 #if PLATFORM_IOS
-					static_cast<FInternetAddrBSDIPv6&>(OutAddr).SetIp(*IPv6SockAddr);
+						NewAddr->SetIp(*IPv6SockAddr);
 #else
-					static_cast<FInternetAddrBSDIPv6&>(OutAddr).SetIp(IPv6SockAddr->sin6_addr);
+						NewAddr->SetIp(IPv6SockAddr->sin6_addr);
 #endif
-					freeaddrinfo(AddrInfo);
-					return SE_NO_ERROR;
+					}
 				}
-			}
-			else if (AddrInfo->ai_family == AF_INET)
-			{
-				sockaddr_in* IPv4SockAddr = reinterpret_cast<sockaddr_in*>(AddrInfo->ai_addr);
-				if (IPv4SockAddr != nullptr)
+				else if (AddrInfo->ai_family == AF_INET)
 				{
-					static_cast<FInternetAddrBSDIPv6&>(OutAddr).SetIp(IPv4SockAddr->sin_addr);
-					freeaddrinfo(AddrInfo);
-					return SE_NO_ERROR;
+					sockaddr_in* IPv4SockAddr = reinterpret_cast<sockaddr_in*>(AddrInfo->ai_addr);
+					if (IPv4SockAddr != nullptr)
+					{
+						NewAddr->SetIp(IPv4SockAddr->sin_addr);
+					}
 				}
+				GetAddrArray.Add(NewAddr);
 			}
 		}
 		freeaddrinfo(AddrInfo);
-		return SE_HOST_NOT_FOUND;
 	}
-	return SocketError;
+#else
+	UE_LOG(LogSockets, Error, TEXT("Platform has no getaddrinfo(), but did not override FSocketSubsystem::GetAddressInfo()"));
+#endif
+	return GetAddrArray;
+}
+
+ESocketErrors FSocketSubsystemBSDIPv6::GetHostByName(const ANSICHAR* HostName, FInternetAddr& OutAddr)
+{
+	TArray<TSharedRef<FInternetAddr> > Results = GetAddressInfo(HostName, true);
+	if (Results.Num() > 0)
+	{
+		OutAddr.SetRawIp(Results[0]->GetRawIp());
+		return SE_NO_ERROR;
+	}
+
+	return SE_HOST_NOT_FOUND;
 }
 
 bool FSocketSubsystemBSDIPv6::GetHostName(FString& HostName)
diff --git a/Engine/Source/Runtime/Sockets/Private/BSDIPv6Sockets/SocketSubsystemBSDIPv6.h b/Engine/Source/Runtime/Sockets/Private/BSDIPv6Sockets/SocketSubsystemBSDIPv6.h
--- a/Engine/Source/Runtime/Sockets/Private/BSDIPv6Sockets/SocketSubsystemBSDIPv6.h
+++ b/Engine/Source/Runtime/Sockets/Private/BSDIPv6Sockets/SocketSubsystemBSDIPv6.h
@@ -29,6 +29,8 @@ public:
 
 	virtual void DestroySocket( class FSocket* Socket ) override;
 
+	virtual TArray<TSharedRef<FInternetAddr> > GetAddressInfo(const ANSICHAR* HostName, bool bResolveAddress, ESocketProtocolFamily ProtocolType = ESocketProtocolFamily::None);
+
 	virtual ESocketErrors GetHostByName( const ANSICHAR* HostName, FInternetAddr& OutAddr ) override;
 
 	virtual bool GetHostName( FString& HostName ) override;
@@ -68,10 +70,8 @@ protected:
 	 */
 	virtual class FSocketBSDIPv6* InternalBSDSocketFactory( SOCKET Socket, ESocketType SocketType, const FString& SocketDescription );
 
-	PRAGMA_DISABLE_DEPRECATION_WARNINGS
 	// allow BSD sockets to use this when creating new sockets from accept() etc
-	friend FSocketBSDIPv6;
-	PRAGMA_ENABLE_DEPRECATION_WARNINGS
+	friend class FSocketBSDIPv6;
 
 	// Used to prevent multiple threads accessing the shared data.
 	FCriticalSection HostByNameSynch;
diff --git a/Engine/Source/Runtime/Sockets/Private/BSDIPv6Sockets/SocketsBSDIPv6.h b/Engine/Source/Runtime/Sockets/Private/BSDIPv6Sockets/SocketsBSDIPv6.h
--- a/Engine/Source/Runtime/Sockets/Private/BSDIPv6Sockets/SocketsBSDIPv6.h
+++ b/Engine/Source/Runtime/Sockets/Private/BSDIPv6Sockets/SocketsBSDIPv6.h
@@ -26,27 +26,18 @@ namespace EIPv6SocketInternalState
 	};
 };
 
-class FSocketBSDIPv6_Base_DEPRECATED
+/**
+ * Implements a BSD network socket.
+ */
+class FSocketBSDIPv6: public FSocket
 {
 protected:
-
-	FSocketBSDIPv6_Base_DEPRECATED(SOCKET InSocket, ISocketSubsystem* InSubsystem) :
-		Socket(InSocket), SocketSubsystem(InSubsystem) { }
-
 	// Holds the BSD socket object.
 	SOCKET Socket;
 
 	// Pointer to the subsystem that created it
 	ISocketSubsystem * SocketSubsystem;
-};
 
-
-/**
- * Implements a BSD network socket.
- */
-class DEPRECATED(4.21, "FSocketBSDIPv6 is deprecated. Move to FSocketBSD.") FSocketBSDIPv6
-	: public FSocket, FSocketBSDIPv6_Base_DEPRECATED
-{
 public:
 
 	/**
@@ -59,7 +50,8 @@ public:
 	 */
 	FSocketBSDIPv6(SOCKET InSocket, ESocketType InSocketType, const FString& InSocketDescription, ISocketSubsystem* InSubsystem) 
 		: FSocket(InSocketType, InSocketDescription)
-		, FSocketBSDIPv6_Base_DEPRECATED(InSocket, InSubsystem)
+		, Socket(InSocket)
+		, SocketSubsystem(InSubsystem)
 	{ }
 
 	/**
