Fix v4 mapped addresses comparison such that we up the address to a common address family and then do the comparison. This uses ntop to determine the address differences.

This fixes hybrid stack issues.

diff --git a/Engine/Source/Runtime/Sockets/Private/BSDSockets/IPAddressBSD.cpp b/Engine/Source/Runtime/Sockets/Private/BSDSockets/IPAddressBSD.cpp
--- a/Engine/Source/Runtime/Sockets/Private/BSDSockets/IPAddressBSD.cpp
+++ b/Engine/Source/Runtime/Sockets/Private/BSDSockets/IPAddressBSD.cpp
@@ -19,7 +19,10 @@ void FInternetAddrBSD::Clear()
 void FInternetAddrBSD::ResetScopeId()
 {
 #if PLATFORM_HAS_BSD_IPV6_SOCKETS
-	((sockaddr_in6*)&Addr)->sin6_scope_id = 0;
+	if (Addr.ss_family == AF_INET6)
+	{
+		((sockaddr_in6*)&Addr)->sin6_scope_id = 0;
+	}
 #endif
 }
 
@@ -296,30 +299,45 @@ FString FInternetAddrBSD::ToString(bool bAppendPort) const
 bool FInternetAddrBSD::operator==(const FInternetAddr& Other) const
 {
 	const FInternetAddrBSD& OtherBSD = static_cast<const FInternetAddrBSD&>(Other);
-	ESocketProtocolFamily CurrentFamily = GetProtocolFamily();
 
-	// Check if the addr families match
-	if (OtherBSD.Addr.ss_family != Addr.ss_family)
+	// If the ports don't match, already fail out.
+	if (GetPort() != OtherBSD.GetPort())
 	{
 		return false;
 	}
 
+	// On IPv6, we'll want to just convert our addresses to IPv6 and then do the comparison.
+	// This fixes issues like mapped addresses not matching when they should.
 #if PLATFORM_HAS_BSD_IPV6_SOCKETS
-	if (CurrentFamily == ESocketProtocolFamily::IPv6)
+
+	auto GetInAddrStruct = [](const FInternetAddrBSD& x) { return (x.Addr.ss_family == AF_INET6) ?
+		(void*)(&(reinterpret_cast<const sockaddr_in6*>(&x.Addr)->sin6_addr)) :
+		(void*)(&(reinterpret_cast<const sockaddr_in*>(&x.Addr)->sin_addr)); };
+
+	char ThisAddressBuffer[INET6_ADDRSTRLEN];
+	char OtherAddressBuffer[INET6_ADDRSTRLEN];
+	FMemory::Memzero(&ThisAddressBuffer, sizeof(ThisAddressBuffer));
+	FMemory::Memzero(&OtherAddressBuffer, sizeof(OtherAddressBuffer));
+
+	if (inet_ntop(AF_INET6, GetInAddrStruct(OtherBSD), OtherAddressBuffer, sizeof(OtherAddressBuffer)) != nullptr &&
+		inet_ntop(AF_INET6, GetInAddrStruct((*this)), ThisAddressBuffer, sizeof(ThisAddressBuffer)) != nullptr)
 	{
-		const sockaddr_in6* OtherBSDAddr = (sockaddr_in6*)&(OtherBSD.Addr);
-		const sockaddr_in6* ThisBSDAddr = ((sockaddr_in6*)&Addr);
-		return memcmp(&(ThisBSDAddr->sin6_addr), &(OtherBSDAddr->sin6_addr), sizeof(in6_addr)) == 0 &&
-			ThisBSDAddr->sin6_port == OtherBSDAddr->sin6_port;
+		int32 AddressLen = FCStringAnsi::Strlen(ThisAddressBuffer);
+		if (AddressLen == FCStringAnsi::Strlen(OtherAddressBuffer))
+		{
+			return FCStringAnsi::Strncmp(ThisAddressBuffer, OtherAddressBuffer, AddressLen) == 0;
+		}
+
+		return false;
 	}
+
 #endif
 
-	if (CurrentFamily == ESocketProtocolFamily::IPv4)
+	if (Addr.ss_family == AF_INET && OtherBSD.Addr.ss_family == AF_INET)
 	{
 		const sockaddr_in* OtherBSDAddr = (sockaddr_in*)&(OtherBSD.Addr);
 		const sockaddr_in* ThisBSDAddr = ((sockaddr_in*)&Addr);
-		return ThisBSDAddr->sin_addr.s_addr == OtherBSDAddr->sin_addr.s_addr &&
-			ThisBSDAddr->sin_port == OtherBSDAddr->sin_port;
+		return ThisBSDAddr->sin_addr.s_addr == OtherBSDAddr->sin_addr.s_addr;
 	}
 
 	return false;
