Log print for the comparision as a quick test to make sure that the ntops are actually comparing correctly.

Will remove after test

diff --git a/Engine/Source/Runtime/Sockets/Private/BSDSockets/IPAddressBSD.cpp b/Engine/Source/Runtime/Sockets/Private/BSDSockets/IPAddressBSD.cpp
--- a/Engine/Source/Runtime/Sockets/Private/BSDSockets/IPAddressBSD.cpp
+++ b/Engine/Source/Runtime/Sockets/Private/BSDSockets/IPAddressBSD.cpp
@@ -323,7 +323,10 @@ bool FInternetAddrBSD::operator==(const FInternetAddr& Other) const
 		inet_ntop(AF_INET6, GetInAddrStruct((*this)), ThisAddressBuffer, sizeof(ThisAddressBuffer)) != nullptr)
 	{
 		int32 AddressLen = FCStringAnsi::Strlen(ThisAddressBuffer);
-		if (AddressLen == FCStringAnsi::Strlen(OtherAddressBuffer))
+		int32 OtherLen = FCStringAnsi::Strlen(OtherAddressBuffer);
+		// THIS IS TEMP: REMOVE THIS AFTER TESTS.
+		UE_LOG(LogSockets, Verbose, TEXT("Comparing the two sockaddrs. Ours[%s, Len: %d] Theirs[%s, Len: %d]"), ThisAddressBuffer, AddressLen, OtherAddressBuffer, OtherLen);
+		if (AddressLen == OtherLen)
 		{
 			return FCStringAnsi::Strncmp(ThisAddressBuffer, OtherAddressBuffer, AddressLen) == 0;
 		}
