* Exposing wrapper for getaddrinfo.
* Adding more fields and new type for gai flags.
* Renamed to GetAddressInfo due to Windows macros.
* Added scope mutators to FInternetAddr
* Properly inject the scope ids on all results when using GetLocalAdapterAddresses
* Fix the endianness for the raw ip mutators
* Fix mapping on address conversion
* New return types for GetAddressInformation

diff --git a/Engine/Plugins/Online/OnlineSubsystemSteam/Source/Private/SocketSubsystemSteam.cpp b/Engine/Plugins/Online/OnlineSubsystemSteam/Source/Private/SocketSubsystemSteam.cpp
--- a/Engine/Plugins/Online/OnlineSubsystemSteam/Source/Private/SocketSubsystemSteam.cpp
+++ b/Engine/Plugins/Online/OnlineSubsystemSteam/Source/Private/SocketSubsystemSteam.cpp
@@ -311,6 +311,28 @@ void FSocketSubsystemSteam::ConnectFailure(FUniqueNetIdSteam& RemoteId)
 	P2PRemove(RemoteId);
 }
 
+/**
+ * Gets the address information of the given hostname and outputs it into an array of resolvable addresses.
+ * It is up to the caller to determine which one is valid for their environment.
+ *
+ * @param HostName string version of the queryable hostname or ip address
+ * @param ServiceName string version of a service name ("http") or a port number ("80")
+ * @param QueryFlags What flags are used in making the getaddrinfo call. Several flags can be used at once by ORing the values together.
+ *                   Platforms are required to translate this value into a the correct flag representation.
+ * @param ProtocolType this is used to limit results from the call. Specifying None will search all valid protocols.
+ *					   Callers will find they rarely have to specify this flag.
+ * @param SocketType What socket type should the results be formatted for. This typically does not change any formatting results and can
+ *                   be safely left to the default value.
+ *
+ * @return the array of results from GetAddrInfo
+ */
+FAddressInfoResult FSocketSubsystemSteam::GetAddressInfo(const TCHAR* HostName, const TCHAR* ServiceName,
+	EAddressInfoFlags QueryFlags, ESocketProtocolFamily ProtocolType,	ESocketType SocketType)
+{
+	UE_LOG_ONLINE(Warning, TEXT("GetAddressInfo is not supported on Steam Sockets"));
+	return FAddressInfoResult(HostName, ServiceName);
+}
+
 /**
  * Does a DNS look up of a host name
  *
diff --git a/Engine/Plugins/Online/OnlineSubsystemSteam/Source/Private/SocketSubsystemSteam.h b/Engine/Plugins/Online/OnlineSubsystemSteam/Source/Private/SocketSubsystemSteam.h
--- a/Engine/Plugins/Online/OnlineSubsystemSteam/Source/Private/SocketSubsystemSteam.h
+++ b/Engine/Plugins/Online/OnlineSubsystemSteam/Source/Private/SocketSubsystemSteam.h
@@ -215,6 +215,26 @@ public:
 	 */
 	virtual void DestroySocket(class FSocket* Socket) override;
 
+	/**
+	 * Gets the address information of the given hostname and outputs it into an array of resolvable addresses.
+	 * It is up to the caller to determine which one is valid for their environment.
+	 *
+	 * @param HostName string version of the queryable hostname or ip address
+	 * @param ServiceName string version of a service name ("http") or a port number ("80")
+	 * @param QueryFlags What flags are used in making the getaddrinfo call. Several flags can be used at once by ORing the values together.
+	 *                   Platforms are required to translate this value into a the correct flag representation.
+	 * @param ProtocolType this is used to limit results from the call. Specifying None will search all valid protocols.
+	 *					   Callers will find they rarely have to specify this flag.
+	 * @param SocketType What socket type should the results be formatted for. This typically does not change any formatting results and can
+	 *                   be safely left to the default value.
+	 *
+	 * @return the array of results from GetAddrInfo
+	 */
+	virtual FAddressInfoResult GetAddressInfo(const TCHAR* HostName, const TCHAR* ServiceName = nullptr,
+		EAddressInfoFlags QueryFlags = EAddressInfoFlags::Default,
+		ESocketProtocolFamily ProtocolType = ESocketProtocolFamily::None,
+		ESocketType SocketType = ESocketType::SOCKTYPE_Unknown) override;
+
 	/**
 	 * Does a DNS look up of a host name
 	 *
diff --git a/Engine/Source/Runtime/Sockets/Private/BSDIPv6Sockets/IPAddressBSDIPv6.cpp b/Engine/Source/Runtime/Sockets/Private/BSDIPv6Sockets/IPAddressBSDIPv6.cpp
--- a/Engine/Source/Runtime/Sockets/Private/BSDIPv6Sockets/IPAddressBSDIPv6.cpp
+++ b/Engine/Source/Runtime/Sockets/Private/BSDIPv6Sockets/IPAddressBSDIPv6.cpp
@@ -24,7 +24,7 @@ void FInternetAddrBSDIPv6::SetRawIp(const TArray<uint8>& RawAddr)
 		// If this address has an interface, we'll copy it over too.
 		if (RawAddr.Num() == 20)
 		{
-			Addr.sin6_scope_id = (RawAddr[16] << 24) | (RawAddr[17] << 16) | (RawAddr[18] << 8) | (RawAddr[19] << 0);
+			Addr.sin6_scope_id = (RawAddr[16] << 0) | (RawAddr[17] << 8) | (RawAddr[18] << 16) | (RawAddr[19] << 24);
 		}
 
 		Addr.sin6_family = AF_INET6;
@@ -48,10 +48,10 @@ TArray<uint8> FInternetAddrBSDIPv6::GetRawIp() const
 	if (Addr.sin6_scope_id != 0)
 	{
 		uint32 RawScopeId = Addr.sin6_scope_id;
-		RawAddressArray.Add((RawScopeId >> 24) & 0xFF);
-		RawAddressArray.Add((RawScopeId >> 16) & 0xFF);
-		RawAddressArray.Add((RawScopeId >> 8) & 0xFF);
 		RawAddressArray.Add((RawScopeId >> 0) & 0xFF);
+		RawAddressArray.Add((RawScopeId >> 8) & 0xFF);
+		RawAddressArray.Add((RawScopeId >> 16) & 0xFF);
+		RawAddressArray.Add((RawScopeId >> 24) & 0xFF);
 	}
 	return RawAddressArray;
 }
diff --git a/Engine/Source/Runtime/Sockets/Private/BSDIPv6Sockets/SocketSubsystemBSDIPv6.cpp b/Engine/Source/Runtime/Sockets/Private/BSDIPv6Sockets/SocketSubsystemBSDIPv6.cpp
--- a/Engine/Source/Runtime/Sockets/Private/BSDIPv6Sockets/SocketSubsystemBSDIPv6.cpp
+++ b/Engine/Source/Runtime/Sockets/Private/BSDIPv6Sockets/SocketSubsystemBSDIPv6.cpp
@@ -82,9 +82,16 @@ void FSocketSubsystemBSDIPv6::DestroySocket(FSocket* Socket)
 	delete Socket;
 }
 
-TArray<TSharedRef<FInternetAddr> > FSocketSubsystemBSDIPv6::GetAddressInfo(const ANSICHAR* HostName, bool bResolveAddress, ESocketProtocolFamily ProtocolType)
+FAddressInfoResult FSocketSubsystemBSDIPv6::GetAddressInfo(const TCHAR* HostName, const TCHAR* ServiceName,
+	EAddressInfoFlags QueryFlags, ESocketProtocolFamily ProtocolType, ESocketType SocketType)
 {
-	TArray<TSharedRef<FInternetAddr> > GetAddrArray;
+	FAddressInfoResult AddrQueryResult = FAddressInfoResult(HostName, ServiceName);
+
+	if (HostName == nullptr && ServiceName == nullptr)
+	{
+		UE_LOG(LogSockets, Warning, TEXT("GetAddressInfo was passed with both a null host and service name, returning empty array"));
+		return AddrQueryResult;
+	}
 
 #if PLATFORM_HAS_BSD_SOCKET_FEATURE_GETADDRINFO
 	addrinfo* AddrInfo = nullptr;
@@ -92,32 +99,40 @@ TArray<TSharedRef<FInternetAddr> > FSocketSubsystemBSDIPv6::GetAddressInfo(const
 	addrinfo HintAddrInfo;
 	FMemory::Memzero(&HintAddrInfo, sizeof(HintAddrInfo));
 	HintAddrInfo.ai_family = AF_UNSPEC;
-	HintAddrInfo.ai_flags = AI_ADDRCONFIG | AI_PASSIVE;
+	HintAddrInfo.ai_flags = GetAddressInfoHintFlag(QueryFlags);
 
-	if (!bResolveAddress)
+	if (SocketType != ESocketType::SOCKTYPE_Unknown)
 	{
-		HintAddrInfo.ai_flags |= AI_NUMERICHOST;
+		bool bIsUDP = (SocketType == ESocketType::SOCKTYPE_Datagram);
+		HintAddrInfo.ai_protocol = bIsUDP ? IPPROTO_UDP : IPPROTO_TCP;
+		HintAddrInfo.ai_socktype = bIsUDP ? SOCK_DGRAM : SOCK_STREAM;
 	}
 
-	int32 ErrorCode = getaddrinfo(HostName, nullptr, &HintAddrInfo, &AddrInfo);
+	int32 ErrorCode = getaddrinfo(TCHAR_TO_UTF8(HostName), TCHAR_TO_UTF8(ServiceName), &HintAddrInfo, &AddrInfo);
 	ESocketErrors SocketError = TranslateGAIErrorCode(ErrorCode);
 	if (SocketError == SE_NO_ERROR)
 	{
+		if (AddrInfo != nullptr && AddrInfo->ai_canonname != nullptr)
+		{
+			AddrQueryResult.CanonicalNameResult = UTF8_TO_TCHAR(AddrInfo->ai_canonname);
+		}
+
 		for (; AddrInfo != nullptr; AddrInfo = AddrInfo->ai_next)
 		{
 			if (AddrInfo->ai_family == AF_INET6 || AddrInfo->ai_family == AF_INET)
 			{
-				TSharedRef<FInternetAddrBSDIPv6> NewAddr = MakeShareable(new FInternetAddrBSDIPv6);
+				TSharedRef<FInternetAddrBSDIPv6> NewAddress = MakeShareable(new FInternetAddrBSDIPv6);
 				if (AddrInfo->ai_family == AF_INET6)
 				{
 					sockaddr_in6* IPv6SockAddr = reinterpret_cast<sockaddr_in6*>(AddrInfo->ai_addr);
 					if (IPv6SockAddr != nullptr)
 					{
 #if PLATFORM_IOS
-						NewAddr->SetIp(*IPv6SockAddr);
+						NewAddress->SetIp(*IPv6SockAddr);
 #else
-						NewAddr->SetIp(IPv6SockAddr->sin6_addr);
+						NewAddress->SetIp(IPv6SockAddr->sin6_addr);
 #endif
+						NewAddress->SetPort(IPv6SockAddr->sin6_port);
 					}
 				}
 				else if (AddrInfo->ai_family == AF_INET)
@@ -125,10 +140,28 @@ TArray<TSharedRef<FInternetAddr> > FSocketSubsystemBSDIPv6::GetAddressInfo(const
 					sockaddr_in* IPv4SockAddr = reinterpret_cast<sockaddr_in*>(AddrInfo->ai_addr);
 					if (IPv4SockAddr != nullptr)
 					{
-						NewAddr->SetIp(IPv4SockAddr->sin_addr);
+						NewAddress->SetIp(IPv4SockAddr->sin_addr);
+						NewAddress->SetPort(IPv4SockAddr->sin_port);
 					}
 				}
-				GetAddrArray.Add(NewAddr);
+
+				ESocketType ResultAddrConfiguration;
+				switch (AddrInfo->ai_protocol)
+				{
+					case IPPROTO_TCP:
+						ResultAddrConfiguration = SOCKTYPE_Streaming;
+						break;
+					case IPPROTO_UDP:
+						ResultAddrConfiguration = SOCKTYPE_Datagram;
+						break;
+					default:
+						ResultAddrConfiguration = SOCKTYPE_Unknown;
+						break;
+				}
+
+				// Everything in this class is stored internally as IPv6
+				AddrQueryResult.Results.Add(FAddressInfoResultData(NewAddress, AddrInfo->ai_addrlen,
+					ESocketProtocolFamily::IPv6, ResultAddrConfiguration));
 			}
 		}
 		freeaddrinfo(AddrInfo);
@@ -136,15 +169,17 @@ TArray<TSharedRef<FInternetAddr> > FSocketSubsystemBSDIPv6::GetAddressInfo(const
 #else
 	UE_LOG(LogSockets, Error, TEXT("Platform has no getaddrinfo(), but did not override FSocketSubsystem::GetAddressInfo()"));
 #endif
-	return GetAddrArray;
+	return AddrQueryResult;
 }
 
 ESocketErrors FSocketSubsystemBSDIPv6::GetHostByName(const ANSICHAR* HostName, FInternetAddr& OutAddr)
 {
-	TArray<TSharedRef<FInternetAddr> > Results = GetAddressInfo(HostName, true);
-	if (Results.Num() > 0)
+	FAddressInfoResult GAIResult = GetAddressInfo(ANSI_TO_TCHAR(HostName), nullptr,
+		EAddressInfoFlags::AllResultsWithMapping | EAddressInfoFlags::OnlyUsableAddresses | EAddressInfoFlags::BindableAddress);
+
+	if (GAIResult.Results.Num() > 0)
 	{
-		OutAddr.SetRawIp(Results[0]->GetRawIp());
+		OutAddr.SetRawIp(GAIResult.Results[0].Address->GetRawIp());
 		return SE_NO_ERROR;
 	}
 
@@ -254,5 +289,61 @@ ESocketErrors FSocketSubsystemBSDIPv6::TranslateErrorCode(int32 Code)
 	ensure(0);
 	return SE_EINVAL;
 }
+
+int32 FSocketSubsystemBSDIPv6::GetAddressInfoHintFlag(EAddressInfoFlags InFlags) const
+{
+	int32 ReturnFlagsCode = 0;
+
+#if PLATFORM_HAS_BSD_SOCKET_FEATURE_GETADDRINFO
+	if (InFlags == EAddressInfoFlags::Default)
+	{
+		return ReturnFlagsCode;
+	}
+
+	if (EnumHasAnyFlags(InFlags, EAddressInfoFlags::NoResolveHost))
+	{
+		ReturnFlagsCode |= AI_NUMERICHOST;
+	}
+
+	if (EnumHasAnyFlags(InFlags, EAddressInfoFlags::NoResolveService))
+	{
+		ReturnFlagsCode |= AI_NUMERICSERV;
+	}
+
+	if (EnumHasAnyFlags(InFlags, EAddressInfoFlags::OnlyUsableAddresses))
+	{
+		ReturnFlagsCode |= AI_ADDRCONFIG;
+	}
+
+	if (EnumHasAnyFlags(InFlags, EAddressInfoFlags::BindableAddress))
+	{
+		ReturnFlagsCode |= AI_PASSIVE;
+	}
+
+	/* This means nothing unless AI_ALL is also specified. */
+	if (EnumHasAnyFlags(InFlags, EAddressInfoFlags::AllowV4MappedAddresses))
+	{
+		ReturnFlagsCode |= AI_V4MAPPED;
+	}
+
+	if (EnumHasAnyFlags(InFlags, EAddressInfoFlags::AllResults))
+	{
+		ReturnFlagsCode |= AI_ALL;
+	}
+
+	if (EnumHasAnyFlags(InFlags, EAddressInfoFlags::CanonicalName))
+	{
+		ReturnFlagsCode |= AI_CANONNAME;
+	}
+
+	if (EnumHasAnyFlags(InFlags, EAddressInfoFlags::FQDomainName))
+	{
+		ReturnFlagsCode |= AI_FQDN;
+	}
+#endif
+
+	return ReturnFlagsCode;
+}
+
 PRAGMA_ENABLE_DEPRECATION_WARNINGS
 #endif
diff --git a/Engine/Source/Runtime/Sockets/Private/BSDIPv6Sockets/SocketSubsystemBSDIPv6.h b/Engine/Source/Runtime/Sockets/Private/BSDIPv6Sockets/SocketSubsystemBSDIPv6.h
--- a/Engine/Source/Runtime/Sockets/Private/BSDIPv6Sockets/SocketSubsystemBSDIPv6.h
+++ b/Engine/Source/Runtime/Sockets/Private/BSDIPv6Sockets/SocketSubsystemBSDIPv6.h
@@ -29,7 +29,10 @@ public:
 
 	virtual void DestroySocket( class FSocket* Socket ) override;
 
-	virtual TArray<TSharedRef<FInternetAddr> > GetAddressInfo(const ANSICHAR* HostName, bool bResolveAddress, ESocketProtocolFamily ProtocolType = ESocketProtocolFamily::None);
+	virtual FAddressInfoResult GetAddressInfo(const TCHAR* HostName, const TCHAR* ServiceName = nullptr,
+		EAddressInfoFlags QueryFlags = EAddressInfoFlags::Default,
+		ESocketProtocolFamily ProtocolType = ESocketProtocolFamily::None,
+		ESocketType SocketType = ESocketType::SOCKTYPE_Unknown) override;
 
 	virtual ESocketErrors GetHostByName( const ANSICHAR* HostName, FInternetAddr& OutAddr ) override;
 
@@ -70,6 +73,11 @@ protected:
 	 */
 	virtual class FSocketBSDIPv6* InternalBSDSocketFactory( SOCKET Socket, ESocketType SocketType, const FString& SocketDescription );
 
+	/**
+	 * Translates an ESocketAddressInfoFlags into a value usable by getaddrinfo
+	 */
+	virtual int32 GetAddressInfoHintFlag(EAddressInfoFlags InFlags) const;
+
 	// allow BSD sockets to use this when creating new sockets from accept() etc
 	friend class FSocketBSDIPv6;
 
diff --git a/Engine/Source/Runtime/Sockets/Private/BSDSockets/IPAddressBSD.cpp b/Engine/Source/Runtime/Sockets/Private/BSDSockets/IPAddressBSD.cpp
--- a/Engine/Source/Runtime/Sockets/Private/BSDSockets/IPAddressBSD.cpp
+++ b/Engine/Source/Runtime/Sockets/Private/BSDSockets/IPAddressBSD.cpp
@@ -26,6 +26,27 @@ void FInternetAddrBSD::ResetScopeId()
 #endif
 }
 
+uint32 FInternetAddrBSD::GetScopeId() const
+{
+#if PLATFORM_HAS_BSD_IPV6_SOCKETS
+	if (Addr.ss_family == AF_INET6)
+	{
+		return ntohl(((sockaddr_in6*)&Addr)->sin6_scope_id);
+	}
+#endif
+	return 0;
+}
+
+void FInternetAddrBSD::SetScopeId(uint32 NewScopeId)
+{
+#if PLATFORM_HAS_BSD_IPV6_SOCKETS
+	if (Addr.ss_family == AF_INET6)
+	{
+		((sockaddr_in6*)&Addr)->sin6_scope_id = htonl(NewScopeId);
+	}
+#endif
+}
+
 void FInternetAddrBSD::SetIp(const sockaddr_storage& IpAddr)
 {
 	// Instead of just replacing the structures entirely, we copy only the ip portion.
@@ -85,22 +106,33 @@ void FInternetAddrBSD::SetIp(const TCHAR* InAddr, bool& bIsValid)
 	}
 }
 
+void FInternetAddrBSD::Set(const sockaddr_storage& AddrData)
+{
+	Addr = AddrData;
+}
+
+void FInternetAddrBSD::Set(const sockaddr_storage& AddrData, SOCKLEN AddrLen)
+{
+	Clear();
+	FMemory::Memcpy(&Addr, &AddrData, (size_t)AddrLen);
+}
+
 TArray<uint8> FInternetAddrBSD::GetRawIp() const
 {
 	TArray<uint8> RawAddressArray;
 	if (Addr.ss_family == AF_INET)
 	{
-		sockaddr_in* IPv4Addr = ((sockaddr_in*)&Addr);
+		const sockaddr_in* IPv4Addr = ((const sockaddr_in*)&Addr);
 		uint32 IntAddr = IPv4Addr->sin_addr.s_addr;
-		RawAddressArray.Add((IntAddr >> 24) & 0xFF);
-		RawAddressArray.Add((IntAddr >> 16) & 0xFF);
-		RawAddressArray.Add((IntAddr >> 8) & 0xFF);
 		RawAddressArray.Add((IntAddr >> 0) & 0xFF);
+		RawAddressArray.Add((IntAddr >> 8) & 0xFF);
+		RawAddressArray.Add((IntAddr >> 16) & 0xFF);
+		RawAddressArray.Add((IntAddr >> 24) & 0xFF);
 	}
 #if PLATFORM_HAS_BSD_IPV6_SOCKETS
 	else if (Addr.ss_family == AF_INET6)
 	{
-		sockaddr_in6* IPv6Addr = ((sockaddr_in6*)&Addr);
+		const sockaddr_in6* IPv6Addr = ((const sockaddr_in6*)&Addr);
 		for (int i = 0; i < 16; ++i)
 		{
 			RawAddressArray.Add(IPv6Addr->sin6_addr.s6_addr[i]);
@@ -110,10 +142,10 @@ TArray<uint8> FInternetAddrBSD::GetRawIp() const
 		if (IPv6Addr->sin6_scope_id != 0)
 		{
 			uint32 RawScopeId = IPv6Addr->sin6_scope_id;
-			RawAddressArray.Add((RawScopeId >> 24) & 0xFF);
-			RawAddressArray.Add((RawScopeId >> 16) & 0xFF);
-			RawAddressArray.Add((RawScopeId >> 8) & 0xFF);
 			RawAddressArray.Add((RawScopeId >> 0) & 0xFF);
+			RawAddressArray.Add((RawScopeId >> 8) & 0xFF);
+			RawAddressArray.Add((RawScopeId >> 16) & 0xFF);
+			RawAddressArray.Add((RawScopeId >> 24) & 0xFF);
 		}
 	}
 #endif
@@ -127,7 +159,7 @@ void FInternetAddrBSD::SetRawIp(const TArray<uint8>& RawAddr)
 	{
 		Addr.ss_family = AF_INET;
 		sockaddr_in* IPv4Addr = ((sockaddr_in*)&Addr);
-		IPv4Addr->sin_addr.s_addr = (RawAddr[0] << 24) | (RawAddr[1] << 16) | (RawAddr[2] << 8) | (RawAddr[3] << 0);
+		IPv4Addr->sin_addr.s_addr =	(RawAddr[0] << 0) | (RawAddr[1] << 8) | (RawAddr[2] << 16) | (RawAddr[3] << 24);
 	}
 #if PLATFORM_HAS_BSD_IPV6_SOCKETS
 	else if(RawAddr.Num() >= 16) // We are IPv6
@@ -141,7 +173,7 @@ void FInternetAddrBSD::SetRawIp(const TArray<uint8>& RawAddr)
 		// If this address has an interface, we'll copy it over too.
 		if (RawAddr.Num() == 20)
 		{
-			IPv6Addr->sin6_scope_id = (RawAddr[16] << 24) | (RawAddr[17] << 16) | (RawAddr[18] << 8) | (RawAddr[19] << 0);
+			IPv6Addr->sin6_scope_id = (RawAddr[16] << 0) | (RawAddr[17] << 8) | (RawAddr[18] << 16) | (RawAddr[19] << 24);
 		}
 
 		Addr.ss_family = AF_INET6;
@@ -153,6 +185,29 @@ void FInternetAddrBSD::SetRawIp(const TArray<uint8>& RawAddr)
 	}
 }
 
+void FInternetAddrBSD::GetIp(uint32& OutAddr) const
+{
+	if (GetProtocolFamily() != ESocketProtocolFamily::IPv4)
+	{
+		OutAddr = 0;
+
+#if PLATFORM_HAS_BSD_IPV6_SOCKETS
+		sockaddr_in6* IPv6Addr = ((sockaddr_in6*)&Addr);
+		if (IN6_IS_ADDR_V4MAPPED(&IPv6Addr->sin6_addr))
+		{
+#if PLATFORM_LITTLE_ENDIAN
+			OutAddr = (IPv6Addr->sin6_addr.s6_addr[12] << 24) | (IPv6Addr->sin6_addr.s6_addr[13] << 16) | (IPv6Addr->sin6_addr.s6_addr[14] << 8) | (IPv6Addr->sin6_addr.s6_addr[15]);
+#else
+			OutAddr = (IPv6Addr->sin6_addr.s6_addr[15] << 24) | (IPv6Addr->sin6_addr.s6_addr[14] << 16) | (IPv6Addr->sin6_addr.s6_addr[13] << 8) | (IPv6Addr->sin6_addr.s6_addr[12]);
+#endif
+		}
+#endif
+		return;
+	}
+
+	OutAddr = ntohl(((sockaddr_in*)&Addr)->sin_addr.s_addr);
+}
+
 void FInternetAddrBSD::SetPort(int32 InPort)
 {
 #if PLATFORM_HAS_BSD_IPV6_SOCKETS
@@ -232,7 +287,6 @@ void FInternetAddrBSD::SetIPv6BroadcastAddress()
 #endif // in6addr_allnodesonlink
 	SetIp(in6addr_allnodesonlink);
 #endif
-	ResetScopeId();
 	SetPort(0);
 }
 
@@ -299,6 +353,13 @@ FString FInternetAddrBSD::ToString(bool bAppendPort) const
 bool FInternetAddrBSD::operator==(const FInternetAddr& Other) const
 {
 	const FInternetAddrBSD& OtherBSD = static_cast<const FInternetAddrBSD&>(Other);
+	ESocketProtocolFamily CurrentFamily = GetProtocolFamily();
+
+	// Check if the addr families match
+	if (OtherBSD.GetProtocolFamily() != CurrentFamily)
+	{
+		return false;
+	}
 
 	// If the ports don't match, already fail out.
 	if (GetPort() != OtherBSD.GetPort())
@@ -306,37 +367,16 @@ bool FInternetAddrBSD::operator==(const FInternetAddr& Other) const
 		return false;
 	}
 
-	// On IPv6, we'll want to just convert our addresses to IPv6 and then do the comparison.
-	// This fixes issues like mapped addresses not matching when they should.
 #if PLATFORM_HAS_BSD_IPV6_SOCKETS
-
-	auto GetInAddrStruct = [](const FInternetAddrBSD& x) { return (x.Addr.ss_family == AF_INET6) ?
-		(void*)(&(reinterpret_cast<const sockaddr_in6*>(&x.Addr)->sin6_addr)) :
-		(void*)(&(reinterpret_cast<const sockaddr_in*>(&x.Addr)->sin_addr)); };
-
-	char ThisAddressBuffer[INET6_ADDRSTRLEN];
-	char OtherAddressBuffer[INET6_ADDRSTRLEN];
-	FMemory::Memzero(&ThisAddressBuffer, sizeof(ThisAddressBuffer));
-	FMemory::Memzero(&OtherAddressBuffer, sizeof(OtherAddressBuffer));
-
-	if (inet_ntop(AF_INET6, GetInAddrStruct(OtherBSD), OtherAddressBuffer, sizeof(OtherAddressBuffer)) != nullptr &&
-		inet_ntop(AF_INET6, GetInAddrStruct((*this)), ThisAddressBuffer, sizeof(ThisAddressBuffer)) != nullptr)
+	if (CurrentFamily == ESocketProtocolFamily::IPv6)
 	{
-		int32 AddressLen = FCStringAnsi::Strlen(ThisAddressBuffer);
-		int32 OtherLen = FCStringAnsi::Strlen(OtherAddressBuffer);
-		// THIS IS TEMP: REMOVE THIS AFTER TESTS.
-		UE_LOG(LogSockets, Verbose, TEXT("Comparing the two sockaddrs. Ours[%s, Len: %d] Theirs[%s, Len: %d]"), ThisAddressBuffer, AddressLen, OtherAddressBuffer, OtherLen);
-		if (AddressLen == OtherLen)
-		{
-			return FCStringAnsi::Strncmp(ThisAddressBuffer, OtherAddressBuffer, AddressLen) == 0;
-		}
-
-		return false;
+		const sockaddr_in6* OtherBSDAddr = (sockaddr_in6*)&(OtherBSD.Addr);
+		const sockaddr_in6* ThisBSDAddr = ((sockaddr_in6*)&Addr);
+		return memcmp(&(ThisBSDAddr->sin6_addr), &(OtherBSDAddr->sin6_addr), sizeof(in6_addr)) == 0;
 	}
-
 #endif
 
-	if (Addr.ss_family == AF_INET && OtherBSD.Addr.ss_family == AF_INET)
+	if (CurrentFamily == ESocketProtocolFamily::IPv4)
 	{
 		const sockaddr_in* OtherBSDAddr = (sockaddr_in*)&(OtherBSD.Addr);
 		const sockaddr_in* ThisBSDAddr = ((sockaddr_in*)&Addr);
@@ -403,7 +443,6 @@ SOCKLEN FInternetAddrBSD::GetStorageSize() const
 		return sizeof(sockaddr_in6);
 #endif
 	}
-
 	return sizeof(sockaddr_storage);
 }
 
diff --git a/Engine/Source/Runtime/Sockets/Private/BSDSockets/IPAddressBSD.h b/Engine/Source/Runtime/Sockets/Private/BSDSockets/IPAddressBSD.h
--- a/Engine/Source/Runtime/Sockets/Private/BSDSockets/IPAddressBSD.h
+++ b/Engine/Source/Runtime/Sockets/Private/BSDSockets/IPAddressBSD.h
@@ -97,13 +97,20 @@ public:
 
 	/**
 	 * Sets the address data via a sockaddr_storage
+	 * Can only really safely set with a sockaddr_storage
 	 *
 	 * @param AddrData the new data to use
 	 */
-	virtual void Set(const sockaddr_storage& AddrData)
-	{
-		Addr = AddrData;
-	}
+	virtual void Set(const sockaddr_storage& AddrData);
+
+	/**
+	 * Sets the address data via a sockaddr_storage.
+	 * Uses memcpy to assign the data.
+	 *
+	 * @param AddrData the new data to use
+	 * @param AddrSize size of the addr data being passed in.
+	 */
+	virtual void Set(const sockaddr_storage& AddrData, SOCKLEN AddrSize);
 
 #if PLATFORM_HAS_BSD_IPV6_SOCKETS
 	/**
@@ -143,16 +150,7 @@ public:
 	 *
 	 * @param OutAddr the out param receiving the ip address
 	 */
-	virtual void GetIp(uint32& OutAddr) const override
-	{
-		if (GetProtocolFamily() != ESocketProtocolFamily::IPv4)
-		{
-			OutAddr = 0;
-			return;
-		}
-
-		OutAddr = ntohl(((sockaddr_in*)&Addr)->sin_addr.s_addr);
-	}
+	virtual void GetIp(uint32& OutAddr) const override;
 
 	/**
 	 * Sets the port number from a host byte order int
@@ -249,6 +247,23 @@ public:
 	 */
 	virtual SOCKLEN GetStorageSize() const;
 	
+	/**
+	 * Sets the scope interface id of the currently held address if this address
+	 * is an IPv6 address internally. If it is not, no data will be assigned.
+	 * The NewScopeId must be in host byte order.
+	 *
+	 * @param NewScopeId the new scope interface id to set this address to
+	 */
+	virtual void SetScopeId(uint32 NewScopeId);
+
+	/**
+	 * Returns the IPv6 scope interface id of the currently held address
+	 * if the address is an IPv6 address.
+	 *
+	 * @return the scope interface id
+	 */
+	virtual uint32 GetScopeId() const;
+	
 	virtual uint32 GetTypeHash() override;
 
 	friend class FSocketBSD;
diff --git a/Engine/Source/Runtime/Sockets/Private/BSDSockets/SocketSubsystemBSD.cpp b/Engine/Source/Runtime/Sockets/Private/BSDSockets/SocketSubsystemBSD.cpp
--- a/Engine/Source/Runtime/Sockets/Private/BSDSockets/SocketSubsystemBSD.cpp
+++ b/Engine/Source/Runtime/Sockets/Private/BSDSockets/SocketSubsystemBSD.cpp
@@ -52,10 +52,12 @@ class FSocket* FSocketSubsystemBSD::CreateSocket(const FName& SocketType, const
 
 ESocketErrors FSocketSubsystemBSD::CreateAddressFromIP(const ANSICHAR* IPAddress, FInternetAddr& OutAddr)
 {
-	TArray<TSharedRef<FInternetAddr> > Results = GetAddrInfo(IPAddress, false);
-	if (Results.Num() > 0)
+	FAddressInfoResult GAIResult = GetAddressInfo(ANSI_TO_TCHAR(IPAddress), nullptr,
+		EAddressInfoFlags::AllResultsWithMapping | EAddressInfoFlags::NoResolveHost | EAddressInfoFlags::OnlyUsableAddresses);
+
+	if (GAIResult.Results.Num() > 0)
 	{
-		OutAddr.SetRawIp(Results[0]->GetRawIp());
+		OutAddr.SetRawIp(GAIResult.Results[0].Address->GetRawIp());
 		return SE_NO_ERROR;
 	}
 	
@@ -64,10 +66,12 @@ ESocketErrors FSocketSubsystemBSD::CreateAddressFromIP(const ANSICHAR* IPAddress
 
 ESocketErrors FSocketSubsystemBSD::GetHostByName(const ANSICHAR* HostName, FInternetAddr& OutAddr)
 {
-	TArray<TSharedRef<FInternetAddr> > Results = GetAddrInfo(HostName, true);
-	if (Results.Num() > 0)
+	FAddressInfoResult GAIResult = GetAddressInfo(ANSI_TO_TCHAR(HostName), nullptr,
+		EAddressInfoFlags::AllResultsWithMapping | EAddressInfoFlags::OnlyUsableAddresses | EAddressInfoFlags::BindableAddress);
+
+	if (GAIResult.Results.Num() > 0)
 	{
-		OutAddr.SetRawIp(Results[0]->GetRawIp());
+		OutAddr.SetRawIp(GAIResult.Results[0].Address->GetRawIp());
 		return SE_NO_ERROR;
 	}
 
@@ -79,9 +83,16 @@ void FSocketSubsystemBSD::DestroySocket(FSocket* Socket)
 	delete Socket;
 }
 
-TArray<TSharedRef<FInternetAddr> > FSocketSubsystemBSD::GetAddrInfo(const ANSICHAR* HostName, bool bResolveAddress, ESocketProtocolFamily ProtocolType)
+FAddressInfoResult FSocketSubsystemBSD::GetAddressInfo(const TCHAR* HostName, const TCHAR* ServiceName,
+	EAddressInfoFlags QueryFlags, ESocketProtocolFamily ProtocolType, ESocketType SocketType)
 {
-	TArray<TSharedRef<FInternetAddr> > GetAddrArray;
+	FAddressInfoResult AddrQueryResult = FAddressInfoResult(HostName, ServiceName);
+
+	if (HostName == nullptr && ServiceName == nullptr)
+	{
+		UE_LOG(LogSockets, Warning, TEXT("GetAddressInfo was passed with both a null host and service name, returning empty array"));
+		return AddrQueryResult;
+	}
 
 #if PLATFORM_HAS_BSD_SOCKET_FEATURE_GETADDRINFO
 	addrinfo* AddrInfo = nullptr;
@@ -93,17 +104,25 @@ TArray<TSharedRef<FInternetAddr> > FSocketSubsystemBSD::GetAddrInfo(const ANSICH
 	addrinfo HintAddrInfo;
 	FMemory::Memzero(&HintAddrInfo, sizeof(HintAddrInfo));
 	HintAddrInfo.ai_family = GetProtocolFamilyValue(ProtocolType);
-	HintAddrInfo.ai_flags = AI_ADDRCONFIG;
+	HintAddrInfo.ai_flags = GetAddressInfoHintFlag(QueryFlags);
 
-	if (!bResolveAddress)
+	if (SocketType != ESocketType::SOCKTYPE_Unknown)
 	{
-		HintAddrInfo.ai_flags |= AI_NUMERICHOST;
+		bool bIsUDP = (SocketType == ESocketType::SOCKTYPE_Datagram);
+		HintAddrInfo.ai_protocol = bIsUDP ? IPPROTO_UDP : IPPROTO_TCP;
+		HintAddrInfo.ai_socktype = bIsUDP ? SOCK_DGRAM : SOCK_STREAM;
 	}
 
-	int32 ErrorCode = getaddrinfo(HostName, nullptr, &HintAddrInfo, &AddrInfo);
+	int32 ErrorCode = getaddrinfo(TCHAR_TO_UTF8(HostName), TCHAR_TO_UTF8(ServiceName), &HintAddrInfo, &AddrInfo);
 	ESocketErrors SocketError = TranslateGAIErrorCode(ErrorCode);
 	if (SocketError == SE_NO_ERROR)
 	{
+		// The canonical name will always be stored in only the first result in a getaddrinfo query
+		if (AddrInfo != nullptr && AddrInfo->ai_canonname != nullptr)
+		{
+			AddrQueryResult.CanonicalNameResult = UTF8_TO_TCHAR(AddrInfo->ai_canonname);
+		}
+
 		for (; AddrInfo != nullptr; AddrInfo = AddrInfo->ai_next)
 		{
 			if (AddrInfo->ai_family == AF_INET || (AddrInfo->ai_family == AF_INET6 && bCanUseIPv6))
@@ -111,22 +130,19 @@ TArray<TSharedRef<FInternetAddr> > FSocketSubsystemBSD::GetAddrInfo(const ANSICH
 				sockaddr_storage* AddrData = reinterpret_cast<sockaddr_storage*>(AddrInfo->ai_addr);
 				if (AddrData != nullptr)
 				{
-					// This might not be good as it locks GetAddrInfo into one class.
-					// TODO: Fix this when GetAddrInfo moves to the base class.
-					// Just convert the sockaddr_storage into the raw ip array format (do we want to make a new helper in this class?)
-					// and push it to SetRawIp
 					TSharedRef<FInternetAddrBSD> NewAddress = MakeShareable(new FInternetAddrBSD);
-					NewAddress->SetIp(*AddrData);
-					GetAddrArray.Add(NewAddress);
+					NewAddress->Set(*AddrData, AddrInfo->ai_addrlen);
+					AddrQueryResult.Results.Add(FAddressInfoResultData(NewAddress, AddrInfo->ai_addrlen,
+						GetProtocolFamilyType(AddrInfo->ai_family), GetSocketType(AddrInfo->ai_protocol)));
 				}
 			}
 		}
 		freeaddrinfo(AddrInfo);
 	}
 #else
-	UE_LOG(LogSockets, Error, TEXT("Platform has no getaddrinfo(), but did not override FSocketSubsystem::GetHostByName()"));
+	UE_LOG(LogSockets, Error, TEXT("Platform has no getaddrinfo(), but did not override FSocketSubsystem::GetAddressInfo()"));
 #endif
-	return GetAddrArray;
+	return AddrQueryResult;
 }
 
 bool FSocketSubsystemBSD::GetHostName(FString& HostName)
@@ -289,4 +305,82 @@ int32 FSocketSubsystemBSD::GetProtocolFamilyValue(ESocketProtocolFamily InProtoc
 	}
 }
 
+ESocketProtocolFamily FSocketSubsystemBSD::GetProtocolFamilyType(int32 InProtocol) const
+{
+	switch (InProtocol)
+	{
+		default:
+		case AF_UNSPEC: return ESocketProtocolFamily::None;
+		case AF_INET:   return ESocketProtocolFamily::IPv4;
+		case AF_INET6:  return ESocketProtocolFamily::IPv6;
+	}
+}
+
+ESocketType FSocketSubsystemBSD::GetSocketType(int32 InSocketType) const
+{
+	switch (InSocketType)
+	{
+		case SOCK_STREAM:
+		case IPPROTO_TCP: return ESocketType::SOCKTYPE_Streaming;
+		case SOCK_DGRAM:
+		case IPPROTO_UDP: return ESocketType::SOCKTYPE_Datagram;
+		default: return SOCKTYPE_Unknown;
+	}
+}
+
+int32 FSocketSubsystemBSD::GetAddressInfoHintFlag(EAddressInfoFlags InFlags) const
+{
+	int32 ReturnFlagsCode = 0;
+
+#if PLATFORM_HAS_BSD_SOCKET_FEATURE_GETADDRINFO
+	if (InFlags == EAddressInfoFlags::Default)
+	{
+		return ReturnFlagsCode;
+	}
+
+	if (EnumHasAnyFlags(InFlags, EAddressInfoFlags::NoResolveHost))
+	{
+		ReturnFlagsCode |= AI_NUMERICHOST;
+	}
+
+	if (EnumHasAnyFlags(InFlags, EAddressInfoFlags::NoResolveService))
+	{
+		ReturnFlagsCode |= AI_NUMERICSERV;
+	}
+
+	if (EnumHasAnyFlags(InFlags, EAddressInfoFlags::OnlyUsableAddresses))
+	{
+		ReturnFlagsCode |= AI_ADDRCONFIG;
+	}
+
+	if (EnumHasAnyFlags(InFlags, EAddressInfoFlags::BindableAddress))
+	{
+		ReturnFlagsCode |= AI_PASSIVE;
+	}
+
+	/* This means nothing unless AI_ALL is also specified. */
+	if (EnumHasAnyFlags(InFlags, EAddressInfoFlags::AllowV4MappedAddresses))
+	{
+		ReturnFlagsCode |= AI_V4MAPPED;
+	}
+
+	if (EnumHasAnyFlags(InFlags, EAddressInfoFlags::AllResults))
+	{
+		ReturnFlagsCode |= AI_ALL;
+	}
+
+	if (EnumHasAnyFlags(InFlags, EAddressInfoFlags::CanonicalName))
+	{
+		ReturnFlagsCode |= AI_CANONNAME;
+	}
+
+	if (EnumHasAnyFlags(InFlags, EAddressInfoFlags::FQDomainName))
+	{
+		ReturnFlagsCode |= AI_FQDN;
+	}
+#endif
+
+	return ReturnFlagsCode;
+}
+
 #endif	//PLATFORM_HAS_BSD_SOCKETS
diff --git a/Engine/Source/Runtime/Sockets/Private/BSDSockets/SocketSubsystemBSD.h b/Engine/Source/Runtime/Sockets/Private/BSDSockets/SocketSubsystemBSD.h
--- a/Engine/Source/Runtime/Sockets/Private/BSDSockets/SocketSubsystemBSD.h
+++ b/Engine/Source/Runtime/Sockets/Private/BSDSockets/SocketSubsystemBSD.h
@@ -5,7 +5,7 @@
 #include "CoreMinimal.h"
 #include "Misc/OutputDeviceRedirector.h"
 #include "SocketSubsystemPackage.h"
-#include "SocketTypes.h"
+#include "SocketSubsystem.h"
 
 #if PLATFORM_HAS_BSD_SOCKETS || PLATFORM_HAS_BSD_IPV6_SOCKETS
 
@@ -24,20 +24,6 @@ public:
 	virtual class FSocket* CreateSocket(const FName& SocketType, const FString& SocketDescription, ESocketProtocolFamily ProtocolType, bool bForceUDP = false);
 	virtual ESocketErrors CreateAddressFromIP(const ANSICHAR* IPAddress, FInternetAddr& OutAddr);
 
-	/**
-	 * Gets the address information of the given hostname and outputs it into an array of resolvable addresses.
-	 * It is up to the caller to determine which one is valid for their environment.
-	 *
-	 * @param HostName string version of the queryable hostname or ip address
-	 * @param bResolveAddress if a blocking query to the DNS provider should be make to determine the ip address of the given hostname.
-	 *						  if this flag is false, the call is asynchronous and will only work with ip addresses
-	 * @param ProtocolType this is used to limit results from the call. Specifying None will search all valid protocols.
-	 *					   Callers will find they rarely have to specify this flag.
-	 *
-	 * @return the array of results from GetAddrInfo
-	 */
-	virtual TArray<TSharedRef<FInternetAddr> > GetAddrInfo(const ANSICHAR* HostName, bool bResolveAddress, ESocketProtocolFamily ProtocolType = ESocketProtocolFamily::None);
-
 	/**
 	 * Specifies the default socket protocol family to use when creating a socket
 	 * without explicitly passing in the protocol type on creation.
@@ -58,6 +44,12 @@ public:
 	virtual TSharedRef<FInternetAddr> CreateInternetAddr(uint32 Address = 0, uint32 Port = 0) override;
 	virtual class FSocket* CreateSocket( const FName& SocketType, const FString& SocketDescription, bool bForceUDP = false ) override;
 	virtual void DestroySocket( class FSocket* Socket ) override;
+
+	virtual FAddressInfoResult GetAddressInfo(const TCHAR* HostName, const TCHAR* ServiceName = nullptr,
+		EAddressInfoFlags QueryFlags = EAddressInfoFlags::Default,
+		ESocketProtocolFamily ProtocolType = ESocketProtocolFamily::None,
+		ESocketType SocketType = ESocketType::SOCKTYPE_Unknown) override;
+
 	virtual ESocketErrors GetHostByName( const ANSICHAR* HostName, FInternetAddr& OutAddr ) override;
 	virtual bool GetHostName( FString& HostName ) override;
 	virtual ESocketErrors GetLastErrorCode() override;
@@ -85,16 +77,31 @@ public:
 
 	virtual ESocketErrors TranslateErrorCode( int32 Code ) override;
 
-protected:
 	/**
-	 * Translates return values of getaddrinfo() to socket error enum
+	 * Translates an ESocketProtocolFamily code into a value usable by raw socket apis.
 	 */
-	ESocketErrors TranslateGAIErrorCode(int32 Code) const;
+	virtual int32 GetProtocolFamilyValue(ESocketProtocolFamily InProtocol) const;
 
 	/**
-	 * Translates an ESocketProtocolFamily code into a value usable by raw socket apis.
+	 * Translates an raw socket family type value into an enum that can be used by the network layer.
 	 */
-	virtual int32 GetProtocolFamilyValue(ESocketProtocolFamily InProtocol) const;
+	virtual ESocketProtocolFamily GetProtocolFamilyType(int32 InProtocol) const;
+
+	/**
+	 * Translates an raw socket protocol type value into an enum that can be used by the network layer.
+	 */
+	virtual ESocketType GetSocketType(int32 InSocketType) const;
+
+	/**
+	 * Translates an ESocketAddressInfoFlags into a value usable by getaddrinfo
+	 */
+	virtual int32 GetAddressInfoHintFlag(EAddressInfoFlags InFlags) const;
+
+protected:
+	/**
+	 * Translates return values of getaddrinfo() to socket error enum
+	 */
+	ESocketErrors TranslateGAIErrorCode(int32 Code) const;
 
 	/**
 	 * Allows a subsystem subclass to create a FSocketBSD sub class.
diff --git a/Engine/Source/Runtime/Sockets/Private/BSDSockets/SocketsBSD.cpp b/Engine/Source/Runtime/Sockets/Private/BSDSockets/SocketsBSD.cpp
--- a/Engine/Source/Runtime/Sockets/Private/BSDSockets/SocketsBSD.cpp
+++ b/Engine/Source/Runtime/Sockets/Private/BSDSockets/SocketsBSD.cpp
@@ -31,14 +31,14 @@ bool FSocketBSD::Close(void)
 bool FSocketBSD::Bind(const FInternetAddr& Addr)
 {
 	const FInternetAddrBSD& BSDAddr = static_cast<const FInternetAddrBSD&>(Addr);
-	return bind(Socket, (sockaddr*)&BSDAddr.Addr, BSDAddr.GetStorageSize()) == 0;
+	return bind(Socket, (const sockaddr*)&(BSDAddr.Addr), BSDAddr.GetStorageSize()) == 0;
 }
 
 
 bool FSocketBSD::Connect(const FInternetAddr& Addr)
 {
 	const FInternetAddrBSD& BSDAddr = static_cast<const FInternetAddrBSD&>(Addr);
-	int32 Return = connect(Socket, (sockaddr*)&BSDAddr.Addr, BSDAddr.GetStorageSize());
+	int32 Return = connect(Socket, (const sockaddr*)&(BSDAddr.Addr), BSDAddr.GetStorageSize());
 	
 	check(SocketSubsystem);
 	ESocketErrors Error = SocketSubsystem->TranslateErrorCode(Return);
@@ -132,7 +132,7 @@ bool FSocketBSD::SendTo(const uint8* Data, int32 Count, int32& BytesSent, const
 {
 	const FInternetAddrBSD& BSDAddr = static_cast<const FInternetAddrBSD&>(Destination);
 	// Write the data and see how much was written
-	BytesSent = sendto(Socket, (const char*)Data, Count, 0, (sockaddr*)&(BSDAddr.Addr), BSDAddr.GetStorageSize());
+	BytesSent = sendto(Socket, (const char*)Data, Count, 0, (const sockaddr*)&(BSDAddr.Addr), BSDAddr.GetStorageSize());
 
 //	NETWORK_PROFILER(FSocket::SendTo(Data,Count,BytesSent,Destination));
 
diff --git a/Engine/Source/Runtime/Sockets/Private/IOS/SocketSubsystemIOS.cpp b/Engine/Source/Runtime/Sockets/Private/IOS/SocketSubsystemIOS.cpp
--- a/Engine/Source/Runtime/Sockets/Private/IOS/SocketSubsystemIOS.cpp
+++ b/Engine/Source/Runtime/Sockets/Private/IOS/SocketSubsystemIOS.cpp
@@ -5,6 +5,7 @@
 #include "Modules/ModuleManager.h"
 #include "BSDSockets/SocketsBSD.h"
 #include "IPAddress.h"
+#include <net/if.h>
 #include <ifaddrs.h>
 #include "SocketsBSDIOS.h"
 #include "IPAddressBSDIOS.h"
@@ -112,17 +113,20 @@ TSharedRef<FInternetAddr> FSocketSubsystemIOS::GetLocalHostAddr(FOutputDevice& O
 			}
 
 			sockaddr_storage* AddrData = reinterpret_cast<sockaddr_storage*>(Travel->ifa_addr);
+			uint32 ScopeInterfaceId = ntohl(if_nametoindex(Travel->ifa_name));
 			if (Travel->ifa_addr->sa_family == AF_INET6)
 			{
 				if (strcmp(Travel->ifa_name, "en0") == 0)
 				{
 					HostAddr->SetIp(*AddrData);
+					HostAddr->SetScopeId(ScopeInterfaceId);
 					bWasWifiSet = true;
 					bWasIPv6Set = true;
 				}
 				else if (!bWasWifiSet && strcmp(Travel->ifa_name, "pdp_ip0") == 0)
 				{
 					HostAddr->SetIp(*AddrData);
+					HostAddr->SetScopeId(ScopeInterfaceId);
 					bWasCellSet = true;
 				}
 			}
@@ -131,11 +135,13 @@ TSharedRef<FInternetAddr> FSocketSubsystemIOS::GetLocalHostAddr(FOutputDevice& O
 				if (strcmp(Travel->ifa_name, "en0") == 0)
 				{
 					HostAddr->SetIp(*AddrData);
+					HostAddr->SetScopeId(ScopeInterfaceId);
 					bWasWifiSet = true;
 				}
 				else if (!bWasWifiSet && strcmp(Travel->ifa_name, "pdp_ip0") == 0)
 				{
 					HostAddr->SetIp(*AddrData);
+					HostAddr->SetScopeId(ScopeInterfaceId);
 					bWasCellSet = true;
 				}
 			}
diff --git a/Engine/Source/Runtime/Sockets/Private/Unix/SocketSubsystemUnix.cpp b/Engine/Source/Runtime/Sockets/Private/Unix/SocketSubsystemUnix.cpp
--- a/Engine/Source/Runtime/Sockets/Private/Unix/SocketSubsystemUnix.cpp
+++ b/Engine/Source/Runtime/Sockets/Private/Unix/SocketSubsystemUnix.cpp
@@ -110,7 +110,7 @@ TSharedRef<FInternetAddr> FSocketSubsystemUnix::GetLocalHostAddr(FOutputDevice&
 		return Addr;
 	}
 
-	TArray<TSharedPtr<FInternetAddr> > ResultArray;
+	TArray<TSharedPtr<FInternetAddr>> ResultArray;
 	if (GetLocalAdapterAddresses(ResultArray))
 	{
 		return ResultArray[0]->Clone();
@@ -124,7 +124,10 @@ TSharedRef<FInternetAddr> FSocketSubsystemUnix::GetLocalHostAddr(FOutputDevice&
 bool FSocketSubsystemUnix::GetLocalAdapterAddresses(TArray<TSharedPtr<FInternetAddr> >& OutAddresses)
 {
 	TSharedRef<FInternetAddr> MultihomeAddress = FSocketSubsystemBSD::CreateInternetAddr();
-	if (GetMultihomeAddress(MultihomeAddress))
+	bool bHasMultihome = GetMultihomeAddress(MultihomeAddress);
+
+	// Multihome addresses should always be the first in the array.
+	if (bHasMultihome)
 	{
 		OutAddresses.Add(MultihomeAddress);
 	}
@@ -144,13 +147,24 @@ bool FSocketSubsystemUnix::GetLocalAdapterAddresses(TArray<TSharedPtr<FInternetA
 
 			uint16 AddrFamily = Travel->ifa_addr->sa_family;
 			// Find any up and non-loopback addresses
-			if ((Travel->ifa_flags & IFF_UP) && 
+			if ((Travel->ifa_flags & IFF_UP) != 0 && 
 				(Travel->ifa_flags & IFF_LOOPBACK) == 0 && 
 				(AddrFamily == AF_INET || AddrFamily == AF_INET6))
 			{
 				TSharedRef<FInternetAddrBSD> NewAddress = MakeShareable(new FInternetAddrBSD);
 				NewAddress->SetIp(*((sockaddr_storage*)Travel->ifa_addr));
-				OutAddresses.Add(NewAddress);
+
+				// Write the scope id if what we found was the multihome address.
+				// Don't write it to our list again though.
+				if (bHasMultihome && NewAddress == MultihomeAddress)
+				{
+					static_cast<FInternetAddrBSD&>(MultihomeAddress.Get()).SetScopeId(ntohl(if_nametoindex(Travel->ifa_name)));
+				}
+				else
+				{
+					NewAddress->SetScopeId((ntohl(if_nametoindex(Travel->ifa_name)));
+					OutAddresses.Add(NewAddress);
+				}
 			}
 		}
 
@@ -159,7 +173,7 @@ bool FSocketSubsystemUnix::GetLocalAdapterAddresses(TArray<TSharedPtr<FInternetA
 	else
 	{
 		UE_LOG(LogSockets, Warning, TEXT("getifaddrs returned result %d"), InterfaceQueryRet);
-		return (OutAddresses.Num() == 0); // if getifaddrs somehow doesn't work but we have multihome, then it's fine.
+		return bHasMultihome; // if getifaddrs somehow doesn't work but we have multihome, then it's fine.
 	}
 
 	return true;
diff --git a/Engine/Source/Runtime/Sockets/Private/Windows/SocketSubsystemWindows.cpp b/Engine/Source/Runtime/Sockets/Private/Windows/SocketSubsystemWindows.cpp
--- a/Engine/Source/Runtime/Sockets/Private/Windows/SocketSubsystemWindows.cpp
+++ b/Engine/Source/Runtime/Sockets/Private/Windows/SocketSubsystemWindows.cpp
@@ -76,6 +76,7 @@ FSocket* FSocketSubsystemWindows::CreateSocket(const FName& SocketType, const FS
 	if (NewSocket != nullptr)
 	{
 		::SetHandleInformation((HANDLE)NewSocket->GetNativeSocket(), HANDLE_FLAG_INHERIT, 0);
+		NewSocket->SetIPv6Only(false);
 	}
 	else
 	{
@@ -134,9 +135,10 @@ bool FSocketSubsystemWindows::GetLocalAdapterAddresses( TArray<TSharedPtr<FInter
 	ULONG Flags = GAA_FLAG_INCLUDE_PREFIX | GAA_FLAG_SKIP_MULTICAST | GAA_FLAG_SKIP_DNS_SERVER | GAA_FLAG_SKIP_FRIENDLY_NAME;
 	ULONG Result;
 	ULONG Size = 0;
+	ULONG Family = (PLATFORM_HAS_BSD_IPV6_SOCKETS) ? AF_UNSPEC : AF_INET;
 
 	// determine the required size of the address list buffer
-	Result = GetAdaptersAddresses(AF_INET, Flags, NULL, NULL, &Size);
+	Result = GetAdaptersAddresses(Family, Flags, NULL, NULL, &Size);
 
 	if (Result != ERROR_BUFFER_OVERFLOW)
 	{
@@ -146,7 +148,7 @@ bool FSocketSubsystemWindows::GetLocalAdapterAddresses( TArray<TSharedPtr<FInter
 	PIP_ADAPTER_ADDRESSES AdapterAddresses = (PIP_ADAPTER_ADDRESSES)FMemory::Malloc(Size);
 
 	// get the actual list of adapters
-	Result = GetAdaptersAddresses(AF_INET, Flags, NULL, AdapterAddresses, &Size);
+	Result = GetAdaptersAddresses(Family, Flags, NULL, AdapterAddresses, &Size);
 	
 	if (Result != ERROR_SUCCESS)
 	{
@@ -168,6 +170,7 @@ bool FSocketSubsystemWindows::GetLocalAdapterAddresses( TArray<TSharedPtr<FInter
 					const sockaddr_storage* RawAddress = (const sockaddr_storage*)(UnicastAddress->Address.lpSockaddr);
 					TSharedRef<FInternetAddrBSD> NewAddress = MakeShareable(new FInternetAddrBSD);
 					NewAddress->SetIp(*RawAddress);
+					NewAddress->SetScopeId(ntohl(AdapterAddress->Ipv6IfIndex));
 					OutAdresses.Add(NewAddress);
 				}
 			}
diff --git a/Engine/Source/Runtime/Sockets/Public/AddressInfoTypes.h b/Engine/Source/Runtime/Sockets/Public/AddressInfoTypes.h
new file mode 100644
--- /dev/null
+++ b/Engine/Source/Runtime/Sockets/Public/AddressInfoTypes.h
@@ -0,0 +1,95 @@
+// Copyright 1998-2018 Epic Games, Inc. All Rights Reserved.
+
+#pragma once
+
+#include "CoreMinimal.h"
+#include "SocketTypes.h"
+
+class FInternetAddr;
+
+enum class EAddressInfoFlags : uint16
+{
+	/* Return all addresses regardless of spec. */
+	AllResults = 1 << 0,
+	/* Don't use DNS resolution of the string (AI_NUMERICHOST),
+	 * this makes the address resolution non-blocking, requiring
+	 * the hostname to be in the form of an ip address already
+	 */
+	NoResolveHost = 1 << 1,
+	/* Don't resolve the service name (requires the service argument to be nullptr
+	 * or a string based representation of a number). This is also-nonblocking.
+	 */
+	NoResolveService = 1 << 2,
+	/* Only return addresses that adapters on this machine can use (AI_ADDRCONFIG) */
+	OnlyUsableAddresses = 1 << 3,
+	/* Return bindable addresses (AI_PASSIVE). Only works if the hostname argument is null. */
+	BindableAddress = 1 << 4,
+	/* Include the canonical name of the host with the results list */
+	CanonicalName = 1 << 5,
+	/* Include the fully qualified domain name of the host with the results list */
+	FQDomainName = 1 << 6,
+	/* Allow for IPv4 mapped IPv6 addresses */
+	AllowV4MappedAddresses = 1 << 7,
+	/* Get all addresses, but return V4 mapped IPv6 addresses */
+	AllResultsWithMapping = AllowV4MappedAddresses | AllResults,
+	/* The default value of a hints flag for the platform (typically just 0) */
+	Default = 0,
+};
+
+ENUM_CLASS_FLAGS(EAddressInfoFlags);
+
+struct FAddressInfoResultData
+{
+	FAddressInfoResultData(TSharedRef<FInternetAddr> InAddr, SIZE_T InAddrLen, ESocketProtocolFamily InProtocol, ESocketType InSocketConfiguration) :
+		AddressProtocol(InProtocol),
+		SocketConfiguration(InSocketConfiguration),
+		AddressLen(InAddrLen),
+		Address(InAddr)
+	{
+	}
+
+	FName GetSocketTypeName() const
+	{
+		switch (SocketConfiguration)
+		{
+		case SOCKTYPE_Datagram:
+			return NAME_DGram;
+			break;
+		case SOCKTYPE_Streaming:
+			return NAME_Stream;
+			break;
+		default:
+		case SOCKTYPE_Unknown:
+			return NAME_None;
+			break;
+		}
+	}
+
+	/* The protocol of the address stored */
+	ESocketProtocolFamily AddressProtocol;
+	/* Streaming or datagram */
+	ESocketType SocketConfiguration;
+	/* Length of the returned address data */
+	SIZE_T AddressLen;
+	/* The address associated with this result */
+	TSharedRef<FInternetAddr> Address;
+};
+
+struct FAddressInfoResult
+{
+	FAddressInfoResult(const TCHAR* InHostName, const TCHAR* InServiceName) :
+		QueryHostName(InHostName),
+		QueryServiceName(InServiceName)
+	{
+	}
+
+	/* The hostname that generated these results */
+	FString QueryHostName;
+	/* The service name that was used in the query */
+	FString QueryServiceName;
+	/* The Canonical Name of the query (empty unless FQDomainName or CanonicalName are specified) */
+	FString CanonicalNameResult;
+	/* The list of results */
+	TArray<FAddressInfoResultData> Results;
+};
+
diff --git a/Engine/Source/Runtime/Sockets/Public/SocketSubsystem.h b/Engine/Source/Runtime/Sockets/Public/SocketSubsystem.h
--- a/Engine/Source/Runtime/Sockets/Public/SocketSubsystem.h
+++ b/Engine/Source/Runtime/Sockets/Public/SocketSubsystem.h
@@ -4,6 +4,7 @@
 
 #include "CoreMinimal.h"
 #include "SocketTypes.h"
+#include "AddressInfoTypes.h"
 
 class Error;
 class FInternetAddr;
@@ -36,8 +37,6 @@ SOCKETS_API DECLARE_LOG_CATEGORY_EXTERN(LogSockets, Log, All);
 	#endif
 #endif
 
-class FInternetAddr;
-
 /**
  * This is the base interface to abstract platform specific sockets API
  * differences.
@@ -101,6 +100,27 @@ public:
 	 */
 	virtual void DestroySocket(class FSocket* Socket) = 0;
 
+	/**
+	 * Gets the address information of the given hostname and outputs it into an array of resolvable addresses.
+	 * It is up to the caller to determine which one is valid for their environment.
+	 *
+	 * @param HostName string version of the queryable hostname or ip address
+	 * @param ServiceName string version of a service name ("http") or a port number ("80")
+	 * @param QueryFlags What flags are used in making the getaddrinfo call. Several flags can be used at once by
+	 *                   bitwise OR-ing the flags together.
+	 *                   Platforms are required to translate this value into a the correct flag representation.
+	 * @param ProtocolType Used to limit results from the call. Specifying None will search all valid protocols.
+	 *					   Callers will find they rarely have to specify this flag.
+	 * @param SocketType What socket type should the results be formatted for. This typically does not change any
+	 *                   formatting results and can be safely left to the default value.
+	 *
+	 * @return the results structure containing the array of address results to this query
+	 */
+	virtual FAddressInfoResult GetAddressInfo(const TCHAR* HostName, const TCHAR* ServiceName = nullptr,
+		EAddressInfoFlags QueryFlags = EAddressInfoFlags::Default,
+		ESocketProtocolFamily ProtocolType = ESocketProtocolFamily::None,
+		ESocketType SocketType = ESocketType::SOCKTYPE_Unknown) = 0;
+
 	/**
 	 * Does a DNS look up of a host name
 	 *
diff --git a/Engine/Source/Runtime/Sockets/Public/SocketTypes.h b/Engine/Source/Runtime/Sockets/Public/SocketTypes.h
--- a/Engine/Source/Runtime/Sockets/Public/SocketTypes.h
+++ b/Engine/Source/Runtime/Sockets/Public/SocketTypes.h
@@ -5,7 +5,7 @@
 #include "CoreMinimal.h"
 
 /** Indicates the socket protocol of socket being used, typically for BSD Sockets */
-enum class ESocketProtocolFamily : uint32
+enum class ESocketProtocolFamily : uint8
 {
 	/** No protocol family specification. Typically defined as AF_UNSPEC */
 	None,
