Fix Unix compilation.

diff --git a/Engine/Source/Runtime/Sockets/Private/Unix/SocketSubsystemUnix.cpp b/Engine/Source/Runtime/Sockets/Private/Unix/SocketSubsystemUnix.cpp
--- a/Engine/Source/Runtime/Sockets/Private/Unix/SocketSubsystemUnix.cpp
+++ b/Engine/Source/Runtime/Sockets/Private/Unix/SocketSubsystemUnix.cpp
@@ -162,7 +162,7 @@ bool FSocketSubsystemUnix::GetLocalAdapterAddresses(TArray<TSharedPtr<FInternetA
 				}
 				else
 				{
-					NewAddress->SetScopeId((ntohl(if_nametoindex(Travel->ifa_name)));
+					NewAddress->SetScopeId(ntohl(if_nametoindex(Travel->ifa_name)));
 					OutAddresses.Add(NewAddress);
 				}
 			}
