Making platforms that don't have FQDN support fall back to canonicalname

diff --git a/Engine/Source/Runtime/Sockets/Private/BSDIPv6Sockets/SocketSubsystemBSDIPv6.cpp b/Engine/Source/Runtime/Sockets/Private/BSDIPv6Sockets/SocketSubsystemBSDIPv6.cpp
--- a/Engine/Source/Runtime/Sockets/Private/BSDIPv6Sockets/SocketSubsystemBSDIPv6.cpp
+++ b/Engine/Source/Runtime/Sockets/Private/BSDIPv6Sockets/SocketSubsystemBSDIPv6.cpp
@@ -338,7 +338,11 @@ int32 FSocketSubsystemBSDIPv6::GetAddressInfoHintFlag(EAddressInfoFlags InFlags)
 
 	if (EnumHasAnyFlags(InFlags, EAddressInfoFlags::FQDomainName))
 	{
+#ifdef AI_FQDN
 		ReturnFlagsCode |= AI_FQDN;
+#else
+		ReturnFlagsCode |= AI_CANONNAME;
+#endif
 	}
 #endif
 
diff --git a/Engine/Source/Runtime/Sockets/Private/BSDSockets/SocketSubsystemBSD.cpp b/Engine/Source/Runtime/Sockets/Private/BSDSockets/SocketSubsystemBSD.cpp
--- a/Engine/Source/Runtime/Sockets/Private/BSDSockets/SocketSubsystemBSD.cpp
+++ b/Engine/Source/Runtime/Sockets/Private/BSDSockets/SocketSubsystemBSD.cpp
@@ -376,7 +376,11 @@ int32 FSocketSubsystemBSD::GetAddressInfoHintFlag(EAddressInfoFlags InFlags) con
 
 	if (EnumHasAnyFlags(InFlags, EAddressInfoFlags::FQDomainName))
 	{
+#ifdef AI_FQDN
 		ReturnFlagsCode |= AI_FQDN;
+#else
+		ReturnFlagsCode |= AI_CANONNAME;
+#endif
 	}
 #endif
 
