* Adding more logging to the Unix adapter queries.
* Fix the adapter query returning success on failure.

diff --git a/Engine/Source/Runtime/Sockets/Private/Unix/SocketSubsystemUnix.cpp b/Engine/Source/Runtime/Sockets/Private/Unix/SocketSubsystemUnix.cpp
--- a/Engine/Source/Runtime/Sockets/Private/Unix/SocketSubsystemUnix.cpp
+++ b/Engine/Source/Runtime/Sockets/Private/Unix/SocketSubsystemUnix.cpp
@@ -104,17 +104,21 @@ TSharedRef<FInternetAddr> FSocketSubsystemUnix::GetLocalHostAddr(FOutputDevice&
 		return Addr;
 	}
 
-	// Grab multihome if it exists.
-	if (GetMultihomeAddress(Addr))
-	{
-		return Addr;
-	}
-
 	TArray<TSharedPtr<FInternetAddr>> ResultArray;
 	if (GetLocalAdapterAddresses(ResultArray))
 	{
+		if (FParse::Param(FCommandLine::Get(), TEXT("PRIMARYNET")))
+		{
+			bCanBindAll = false;
+		}
+
+		UE_LOG(LogSockets, Verbose, TEXT("Local address is %s"), *(ResultArray[0]->ToString(false)));
 		return ResultArray[0]->Clone();
 	}
+	else
+	{
+		UE_LOG(LogSockets, Warning, TEXT("GetLocalAdapterAddresses had no results!"));
+	}
 
 	// Fall back to this.
 	Addr->SetAnyAddress();
@@ -176,5 +180,5 @@ bool FSocketSubsystemUnix::GetLocalAdapterAddresses(TArray<TSharedPtr<FInternetA
 		return bHasMultihome; // if getifaddrs somehow doesn't work but we have multihome, then it's fine.
 	}
 
-	return true;
+	return (OutAddresses.Num() > 0);
 }
