Fix the memory leak caused by freeing the wrong sublist.

diff --git a/Engine/Source/Runtime/Sockets/Private/BSDIPv6Sockets/SocketSubsystemBSDIPv6.cpp b/Engine/Source/Runtime/Sockets/Private/BSDIPv6Sockets/SocketSubsystemBSDIPv6.cpp
--- a/Engine/Source/Runtime/Sockets/Private/BSDIPv6Sockets/SocketSubsystemBSDIPv6.cpp
+++ b/Engine/Source/Runtime/Sockets/Private/BSDIPv6Sockets/SocketSubsystemBSDIPv6.cpp
@@ -112,6 +112,7 @@ FAddressInfoResult FSocketSubsystemBSDIPv6::GetAddressInfo(const TCHAR* HostName
 	ESocketErrors SocketError = TranslateGAIErrorCode(ErrorCode);
 	if (SocketError == SE_NO_ERROR)
 	{
+		addrinfo* AddrInfoHead = AddrInfo;
 		if (AddrInfo != nullptr && AddrInfo->ai_canonname != nullptr)
 		{
 			AddrQueryResult.CanonicalNameResult = UTF8_TO_TCHAR(AddrInfo->ai_canonname);
@@ -164,7 +165,7 @@ FAddressInfoResult FSocketSubsystemBSDIPv6::GetAddressInfo(const TCHAR* HostName
 					ESocketProtocolFamily::IPv6, ResultAddrConfiguration));
 			}
 		}
-		freeaddrinfo(AddrInfo);
+		freeaddrinfo(AddrInfoHead);
 	}
 #else
 	UE_LOG(LogSockets, Error, TEXT("Platform has no getaddrinfo(), but did not override FSocketSubsystem::GetAddressInfo()"));
diff --git a/Engine/Source/Runtime/Sockets/Private/BSDSockets/SocketSubsystemBSD.cpp b/Engine/Source/Runtime/Sockets/Private/BSDSockets/SocketSubsystemBSD.cpp
--- a/Engine/Source/Runtime/Sockets/Private/BSDSockets/SocketSubsystemBSD.cpp
+++ b/Engine/Source/Runtime/Sockets/Private/BSDSockets/SocketSubsystemBSD.cpp
@@ -117,6 +117,7 @@ FAddressInfoResult FSocketSubsystemBSD::GetAddressInfo(const TCHAR* HostName, co
 	ESocketErrors SocketError = TranslateGAIErrorCode(ErrorCode);
 	if (SocketError == SE_NO_ERROR)
 	{
+		addrinfo* AddrInfoHead = AddrInfo;
 		// The canonical name will always be stored in only the first result in a getaddrinfo query
 		if (AddrInfo != nullptr && AddrInfo->ai_canonname != nullptr)
 		{
@@ -137,7 +138,7 @@ FAddressInfoResult FSocketSubsystemBSD::GetAddressInfo(const TCHAR* HostName, co
 				}
 			}
 		}
-		freeaddrinfo(AddrInfo);
+		freeaddrinfo(AddrInfoHead);
 	}
 #else
 	UE_LOG(LogSockets, Error, TEXT("Platform has no getaddrinfo(), but did not override FSocketSubsystem::GetAddressInfo()"));
