Adding a new compare function for FInternetAddrs. This CompareEndpoints function can handle protocol differences instead of straight checking the structure data.

diff --git a/Engine/Plugins/Online/OnlineSubsystemUtils/Source/OnlineSubsystemUtils/Private/IpNetDriver.cpp b/Engine/Plugins/Online/OnlineSubsystemUtils/Source/OnlineSubsystemUtils/Private/IpNetDriver.cpp
--- a/Engine/Plugins/Online/OnlineSubsystemUtils/Source/OnlineSubsystemUtils/Private/IpNetDriver.cpp
+++ b/Engine/Plugins/Online/OnlineSubsystemUtils/Source/OnlineSubsystemUtils/Private/IpNetDriver.cpp
@@ -292,7 +292,7 @@ void UIpNetDriver::TickDispatch(float DeltaTime)
 
 					if (MyServerConnection)
 					{
-						if (*MyServerConnection->RemoteAddr == *FromAddr)
+						if (MyServerConnection->RemoteAddr->CompareEndpoints(*FromAddr))
 						{
 							Connection = MyServerConnection;
 						}
@@ -353,7 +353,7 @@ void UIpNetDriver::TickDispatch(float DeltaTime)
 		// Figure out which socket the received data came from.
 		if (MyServerConnection)
 		{
-			if ((*MyServerConnection->RemoteAddr == *FromAddr))
+			if (MyServerConnection->RemoteAddr->CompareEndpoints(*FromAddr))
 			{
 				Connection = MyServerConnection;
 			}
@@ -371,7 +371,7 @@ void UIpNetDriver::TickDispatch(float DeltaTime)
 
 			Connection = (Result != nullptr ? Cast<UIpConnection>(*Result) : nullptr);
 
-			check(Connection == nullptr || *Connection->RemoteAddr == *FromAddr);
+			check(Connection == nullptr || Connection->RemoteAddr->CompareEndpoints(*FromAddr));
 		}
 
 
diff --git a/Engine/Source/Runtime/Sockets/Private/BSDSockets/IPAddressBSD.cpp b/Engine/Source/Runtime/Sockets/Private/BSDSockets/IPAddressBSD.cpp
--- a/Engine/Source/Runtime/Sockets/Private/BSDSockets/IPAddressBSD.cpp
+++ b/Engine/Source/Runtime/Sockets/Private/BSDSockets/IPAddressBSD.cpp
@@ -10,6 +10,62 @@ FInternetAddrBSD::FInternetAddrBSD()
 	Clear();
 }
 
+bool FInternetAddrBSD::CompareEndpoints(const FInternetAddr& InAddr) const
+{
+	const FInternetAddrBSD& OtherBSD = static_cast<const FInternetAddrBSD&>(InAddr);
+	if (GetPort() != OtherBSD.GetPort())
+	{
+		return false;
+	}
+
+	// If we share the same addresses, then just let the comparison operator take over.
+	if (Addr.ss_family == OtherBSD.Addr.ss_family)
+	{
+		return *this == InAddr;
+	}
+	else if (Addr.ss_family == AF_INET || OtherBSD.Addr.ss_family == AF_INET)
+	{
+#if PLATFORM_HAS_BSD_IPV6_SOCKETS
+		// To handle mapped addresses, we want to raise one of the addresses to IPv6 and then do the comparison.
+		const in6_addr* IPv6Addr = (Addr.ss_family == AF_INET6) ? &((sockaddr_in6*)&Addr)->sin6_addr
+			: &((sockaddr_in6*)&(OtherBSD.Addr))->sin6_addr;
+
+		// Figure out which address is the one that needs to be raised to IPv6
+		const in_addr* IPv4Addr = (Addr.ss_family == AF_INET) ? &((sockaddr_in*)&Addr)->sin_addr
+			: &((sockaddr_in*)&(OtherBSD.Addr))->sin_addr;
+
+		// Check special addresses first (Multicast, Any, Loopback)
+		if ((IN6_IS_ADDR_MC_LINKLOCAL(IPv6Addr) && IN_MULTICAST(IPv4Addr->s_addr)) ||
+			(IN6_IS_ADDR_UNSPECIFIED(IPv6Addr) && IPv4Addr->s_addr == INADDR_ANY) ||
+			(IN6_IS_ADDR_LOOPBACK(IPv6Addr) && IPv4Addr->s_addr == INADDR_LOOPBACK))
+		{
+			return true;
+		}
+
+		// If we're not IPv4 mapped already, then we're not able to be compared
+		// and should early out
+		if (!IN6_IS_ADDR_V4MAPPED(IPv6Addr))
+		{
+			return false;
+		}
+
+		in6_addr ConvertedAddrData;
+		FMemory::Memzero(ConvertedAddrData);
+		ConvertedAddrData.s6_addr[10] = 0xff;
+		ConvertedAddrData.s6_addr[11] = 0xff;
+		ConvertedAddrData.s6_addr[12] = (static_cast<uint32>(IPv4Addr->s_addr) & 0xFF);
+		ConvertedAddrData.s6_addr[13] = ((static_cast<uint32>(IPv4Addr->s_addr) >> 8) & 0xFF);
+		ConvertedAddrData.s6_addr[14] = ((static_cast<uint32>(IPv4Addr->s_addr) >> 16) & 0xFF);
+		ConvertedAddrData.s6_addr[15] = ((static_cast<uint32>(IPv4Addr->s_addr) >> 24) & 0xFF);
+
+		return memcmp(&(ConvertedAddrData), IPv6Addr, sizeof(in6_addr)) == 0;
+#else
+		return false;
+#endif
+	}
+	return false;
+}
+
 void FInternetAddrBSD::Clear()
 {
 	FMemory::Memzero(&Addr, sizeof(Addr));
@@ -369,7 +425,7 @@ bool FInternetAddrBSD::operator==(const FInternetAddr& Other) const
 
 #if PLATFORM_HAS_BSD_IPV6_SOCKETS
 	if (CurrentFamily == ESocketProtocolFamily::IPv6)
-	{
+		{
 		const sockaddr_in6* OtherBSDAddr = (sockaddr_in6*)&(OtherBSD.Addr);
 		const sockaddr_in6* ThisBSDAddr = ((sockaddr_in6*)&Addr);
 		return memcmp(&(ThisBSDAddr->sin6_addr), &(OtherBSDAddr->sin6_addr), sizeof(in6_addr)) == 0;
diff --git a/Engine/Source/Runtime/Sockets/Private/BSDSockets/IPAddressBSD.h b/Engine/Source/Runtime/Sockets/Private/BSDSockets/IPAddressBSD.h
--- a/Engine/Source/Runtime/Sockets/Private/BSDSockets/IPAddressBSD.h
+++ b/Engine/Source/Runtime/Sockets/Private/BSDSockets/IPAddressBSD.h
@@ -51,6 +51,16 @@ public:
 	 */
 	FInternetAddrBSD();
 
+	/**
+	 * Compares FInternetAddrs together, comparing the logical net addresses (endpoints)
+	 * of the data stored, rather than doing a memory comparison like the equality operator does.
+	 *
+	 * @Param InAddr The address to compare with.
+	 *
+	 * @return true if the endpoint stored in this FInternetAddr is the same as the input.
+	 */
+	virtual bool CompareEndpoints(const FInternetAddr& InAddr) const override;
+
 	/**
 	 * Sets the ip address from a raw network byte order array.
 	 *
diff --git a/Engine/Source/Runtime/Sockets/Public/IPAddress.h b/Engine/Source/Runtime/Sockets/Public/IPAddress.h
--- a/Engine/Source/Runtime/Sockets/Public/IPAddress.h
+++ b/Engine/Source/Runtime/Sockets/Public/IPAddress.h
@@ -23,6 +23,21 @@ public:
 	{
 	}
 
+	/**
+	 * Compares FInternetAddrs together, comparing the logical net addresses (endpoints)
+	 * of the data stored, rather than doing a memory comparison like the equality operator does.
+	 * If it is not explicitly implemented, this falls back to just doing the same behavior of the
+	 * comparison operator.
+	 *
+	 * @Param InAddr The address to compare with.
+	 *
+	 * @return true if the endpoint stored in this FInternetAddr is the same as the input.
+	 */
+	virtual bool CompareEndpoints(const FInternetAddr& InAddr) const
+	{
+		return *this == InAddr;
+	}
+
 	/**
 	 * Sets the ip address from a host byte order uint32
 	 *
