Adding more logging to GAI so that we can determine order of addresses.

diff --git a/Engine/Source/Runtime/Sockets/Private/BSDSockets/SocketSubsystemBSD.cpp b/Engine/Source/Runtime/Sockets/Private/BSDSockets/SocketSubsystemBSD.cpp
--- a/Engine/Source/Runtime/Sockets/Private/BSDSockets/SocketSubsystemBSD.cpp
+++ b/Engine/Source/Runtime/Sockets/Private/BSDSockets/SocketSubsystemBSD.cpp
@@ -20,7 +20,7 @@ FSocket* FSocketSubsystemBSD::CreateSocket(const FName& SocketType, const FStrin
 	return CreateSocket(SocketType, SocketDescription, GetDefaultSocketProtocolFamily(), bForceUDP);
 }
 
-class FSocket* FSocketSubsystemBSD::CreateSocket(const FName& SocketType, const FString& SocketDescription, ESocketProtocolFamily ProtocolType, bool bForceUDP)
+FSocket* FSocketSubsystemBSD::CreateSocket(const FName& SocketType, const FString& SocketDescription, ESocketProtocolFamily ProtocolType, bool bForceUDP)
 {
 	SOCKET Socket = INVALID_SOCKET;
 	FSocket* NewSocket = nullptr;
@@ -90,7 +90,7 @@ FAddressInfoResult FSocketSubsystemBSD::GetAddressInfo(const TCHAR* HostName, co
 
 	if (HostName == nullptr && ServiceName == nullptr)
 	{
-		UE_LOG(LogSockets, Warning, TEXT("GetAddressInfo was passed with both a null host and service name, returning empty array"));
+		UE_LOG(LogSockets, Warning, TEXT("GetAddressInfo was passed with both a null host and service, returning empty result"));
 		return AddrQueryResult;
 	}
 
@@ -115,6 +115,8 @@ FAddressInfoResult FSocketSubsystemBSD::GetAddressInfo(const TCHAR* HostName, co
 
 	int32 ErrorCode = getaddrinfo(TCHAR_TO_UTF8(HostName), TCHAR_TO_UTF8(ServiceName), &HintAddrInfo, &AddrInfo);
 	ESocketErrors SocketError = TranslateGAIErrorCode(ErrorCode);
+
+	UE_LOG(LogSockets, Verbose, TEXT("Executed getaddrinfo with HostName: %s Return: %d"), HostName, ErrorCode);
 	if (SocketError == SE_NO_ERROR)
 	{
 		addrinfo* AddrInfoHead = AddrInfo;
@@ -135,11 +137,17 @@ FAddressInfoResult FSocketSubsystemBSD::GetAddressInfo(const TCHAR* HostName, co
 					NewAddress->Set(*AddrData, AddrInfo->ai_addrlen);
 					AddrQueryResult.Results.Add(FAddressInfoResultData(NewAddress, AddrInfo->ai_addrlen,
 						GetProtocolFamilyType(AddrInfo->ai_family), GetSocketType(AddrInfo->ai_protocol)));
+
+					UE_LOG(LogSockets, Verbose, TEXT("# Family: %s Address: %s"), ((AddrInfo->ai_family == AF_INET) ? TEXT("IPv4") : TEXT("IPv6")), *(NewAddress->ToString(false)));
 				}
 			}
 		}
 		freeaddrinfo(AddrInfoHead);
 	}
+	else
+	{
+		UE_LOG(LogSockets, Warning, TEXT("GetAddressInfo failed to resolve host with error %s [%d]"), GetSocketError(SocketError), ErrorCode);
+	}
 #else
 	UE_LOG(LogSockets, Error, TEXT("Platform has no getaddrinfo(), but did not override FSocketSubsystem::GetAddressInfo()"));
 #endif
