Fix Unix compiler errors

diff --git a/Engine/Source/Runtime/Sockets/Private/Unix/SocketSubsystemUnix.cpp b/Engine/Source/Runtime/Sockets/Private/Unix/SocketSubsystemUnix.cpp
--- a/Engine/Source/Runtime/Sockets/Private/Unix/SocketSubsystemUnix.cpp
+++ b/Engine/Source/Runtime/Sockets/Private/Unix/SocketSubsystemUnix.cpp
@@ -4,6 +4,7 @@
 #include "Misc/CommandLine.h"
 #include "SocketSubsystemModule.h"
 #include "BSDSockets/IPAddressBSD.h"
+#include "BSDSockets/SocketsBSD.h"
 
 #include <ifaddrs.h>
 #include <net/if.h>
