* Make the new IPv6 changes act the same way that it does on current 4.20. This allows for a much smoother adoption
* Add new functions to specify the type of special address the user wants.
* Fix the iOS broadcast addressing to determine IPv6 interface scope.
* Additional logging for iOS to determine how it's picking up interfaces
* Lock hybrid platforms to IPv4 for this version until full compatibility is reached that we can deprecate unsupported flows.

diff --git a/Engine/Source/Runtime/Core/Public/Mac/MacPlatform.h b/Engine/Source/Runtime/Core/Public/Mac/MacPlatform.h
--- a/Engine/Source/Runtime/Core/Public/Mac/MacPlatform.h
+++ b/Engine/Source/Runtime/Core/Public/Mac/MacPlatform.h
@@ -37,7 +37,6 @@ typedef FMacPlatformTypes FPlatformTypes;
 #define PLATFORM_COMPILER_DISTINGUISHES_INT_AND_LONG			1
 #define PLATFORM_TCHAR_IS_4_BYTES					1
 #define PLATFORM_HAS_BSD_TIME						1
-#define PLATFORM_HAS_BSD_IPV6_SOCKETS				1
 //#define PLATFORM_USE_PTHREADS						1
 #define PLATFORM_MAX_FILEPATH_LENGTH				MAX_PATH
 #define PLATFORM_SUPPORTS_TBB						1
diff --git a/Engine/Source/Runtime/Core/Public/Unix/UnixPlatform.h b/Engine/Source/Runtime/Core/Public/Unix/UnixPlatform.h
--- a/Engine/Source/Runtime/Core/Public/Unix/UnixPlatform.h
+++ b/Engine/Source/Runtime/Core/Public/Unix/UnixPlatform.h
@@ -50,7 +50,6 @@ typedef FUnixPlatformTypes FPlatformTypes;
 #define PLATFORM_HAS_NO_EPROCLIM					1
 #define PLATFORM_HAS_BSD_SOCKET_FEATURE_IOCTL		1
 #define PLATFORM_HAS_BSD_SOCKET_FEATURE_MSG_DONTWAIT	1
-#define PLATFORM_HAS_BSD_IPV6_SOCKETS				1
 #define PLATFORM_SUPPORTS_STACK_SYMBOLS				1
 
 #define PLATFORM_ENABLE_POPCNT_INTRINSIC 1
diff --git a/Engine/Source/Runtime/Sockets/Private/BSDSockets/IPAddressBSD.cpp b/Engine/Source/Runtime/Sockets/Private/BSDSockets/IPAddressBSD.cpp
--- a/Engine/Source/Runtime/Sockets/Private/BSDSockets/IPAddressBSD.cpp
+++ b/Engine/Source/Runtime/Sockets/Private/BSDSockets/IPAddressBSD.cpp
@@ -291,47 +291,55 @@ int32 FInternetAddrBSD::GetPort() const
 
 void FInternetAddrBSD::SetAnyAddress()
 {
-#if PLATFORM_HAS_BSD_IPV6_SOCKETS
-	SetAnyIPv6Address();
-#else
-	SetAnyIPv4Address();
-#endif
+	FSocketSubsystemBSD* SocketSubsystem = static_cast<FSocketSubsystemBSD*>(ISocketSubsystem::Get(PLATFORM_SOCKETSUBSYSTEM));
+	if (SocketSubsystem)
+	{
+		SetAnyAddress(SocketSubsystem->GetDefaultSocketProtocolFamily());
+	}
+	else
+	{
+		UE_LOG(LogSockets, Warning, TEXT("Could not determine the default protocol to use in SetAnyAddress!"));
+	}
 }
 
 void FInternetAddrBSD::SetAnyIPv4Address()
 {
+	Clear();
 	((sockaddr_in*)&Addr)->sin_addr.s_addr = htonl(INADDR_ANY);
 	Addr.ss_family = AF_INET;
-	SetPort(0);
 }
 
 void FInternetAddrBSD::SetAnyIPv6Address()
 {
+	Clear();
 #if PLATFORM_HAS_BSD_IPV6_SOCKETS
 	SetIp(in6addr_any);
 #endif
-	ResetScopeId();
-	SetPort(0);
 }
 
 void FInternetAddrBSD::SetBroadcastAddress()
 {
-#if PLATFORM_HAS_BSD_IPV6_SOCKETS
-	SetIPv6BroadcastAddress();
-#else
-	SetIPv4BroadcastAddress();
-#endif
+	FSocketSubsystemBSD* SocketSubsystem = static_cast<FSocketSubsystemBSD*>(ISocketSubsystem::Get(PLATFORM_SOCKETSUBSYSTEM));
+	if (SocketSubsystem)
+	{
+		SetBroadcastAddress(SocketSubsystem->GetDefaultSocketProtocolFamily());
+	}
+	else
+	{
+		UE_LOG(LogSockets, Warning, TEXT("Could not determine the default protocol to use in SetBroadcastAddress!"));
+	}
 }
 
 void FInternetAddrBSD::SetIPv4BroadcastAddress()
 {
+	Clear();
 	((sockaddr_in*)&Addr)->sin_addr.s_addr = htonl(INADDR_BROADCAST);
 	Addr.ss_family = AF_INET;
-	SetPort(0);
 }
 
 void FInternetAddrBSD::SetIPv6BroadcastAddress()
 {
+	Clear();
 #if PLATFORM_HAS_BSD_IPV6_SOCKETS
 	// broadcast means something different in IPv6, but this is a rough equivalent
 #ifndef in6addr_allnodesonlink
@@ -343,32 +351,34 @@ void FInternetAddrBSD::SetIPv6BroadcastAddress()
 #endif // in6addr_allnodesonlink
 	SetIp(in6addr_allnodesonlink);
 #endif
-	SetPort(0);
 }
 
 void FInternetAddrBSD::SetLoopbackAddress()
 {
-#if PLATFORM_HAS_BSD_IPV6_SOCKETS
-	SetIPv6LoopbackAddress();
-#else
-	SetIPv4LoopbackAddress();
-#endif
+	FSocketSubsystemBSD* SocketSubsystem = static_cast<FSocketSubsystemBSD*>(ISocketSubsystem::Get(PLATFORM_SOCKETSUBSYSTEM));
+	if (SocketSubsystem)
+	{
+		SetLoopbackAddress(SocketSubsystem->GetDefaultSocketProtocolFamily());
+	}
+	else
+	{
+		UE_LOG(LogSockets, Warning, TEXT("Could not determine the default protocol to use in SetLoopbackAddress!"));
+	}
 }
 
 void FInternetAddrBSD::SetIPv4LoopbackAddress()
 {
+	Clear();
 	((sockaddr_in*)&Addr)->sin_addr.s_addr = htonl(INADDR_LOOPBACK);
 	Addr.ss_family = AF_INET;
-	SetPort(0);
 }
 
 void FInternetAddrBSD::SetIPv6LoopbackAddress()
 {
+	Clear();
 #if PLATFORM_HAS_BSD_IPV6_SOCKETS
 	SetIp(in6addr_loopback);
 #endif
-	ResetScopeId();
-	SetPort(0);
 }
 
 FString FInternetAddrBSD::ToString(bool bAppendPort) const
diff --git a/Engine/Source/Runtime/Sockets/Private/BSDSockets/IPAddressBSD.h b/Engine/Source/Runtime/Sockets/Private/BSDSockets/IPAddressBSD.h
--- a/Engine/Source/Runtime/Sockets/Private/BSDSockets/IPAddressBSD.h
+++ b/Engine/Source/Runtime/Sockets/Private/BSDSockets/IPAddressBSD.h
@@ -174,45 +174,102 @@ public:
 
 	/**
 	 * Sets the address structure to be bindable to any IP address.
-	 * IPv6 will take precedence.
+	 * Uses the value returned from GetDefaultSocketProtocolFamily to determine default
+	 * addresses to be used.
 	 *
 	 * To skip assumptions, you can call the designated version explicitly below.
 	 */
 	virtual void SetAnyAddress() override;
 
+	/**
+	 * Sets the address structure to be bindable to any IP address
+	 * for the protocol family specified.
+	 *
+	 * @param ForProtocol the protocol to set the address for
+	 */
+	virtual void SetAnyAddress(ESocketProtocolFamily ForProtocol)
+	{
+		if (ForProtocol == ESocketProtocolFamily::IPv4)
+		{
+			SetAnyIPv4Address();
+		}
+		else
+		{
+			SetAnyIPv6Address();
+		}
+	}
+
 	/** Explicit set to any IPv4 address */
-	void SetAnyIPv4Address();
+	virtual void SetAnyIPv4Address();
 
 	/** Explicit set to any IPv6 address */
-	void SetAnyIPv6Address();
+	virtual void SetAnyIPv6Address();
 
 	/**
 	 * Sets the address structure to be bound to the multicast ip address.
-	 * IPv6 will take precedence.
+	 * Uses the value returned from GetDefaultSocketProtocolFamily to determine default
+	 * addresses to be used.
 	 *
 	 * To skip assumptions, you can call the designated version explicitly below.
 	 */
 	virtual void SetBroadcastAddress() override;
 
+	/**
+	 * Sets the address structure to be bound to the multicast ip address
+	 * for the protocol family specified.
+	 *
+	 * @param ForProtocol the protocol to set the address for
+	 */
+	virtual void SetBroadcastAddress(ESocketProtocolFamily ForProtocol)
+	{
+		if (ForProtocol == ESocketProtocolFamily::IPv4)
+		{
+			SetIPv4BroadcastAddress();
+		}
+		else
+		{
+			SetIPv6BroadcastAddress();
+		}
+	}
+
 	/** Explicit set to multicast IPv4 address */
-	void SetIPv4BroadcastAddress();
+	virtual void SetIPv4BroadcastAddress();
 
 	/** Explicit set to multicast IPv6 address */
-	void SetIPv6BroadcastAddress();
+	virtual void SetIPv6BroadcastAddress();
 
 	/**
 	 * Sets the address structure to be bound to the loopback ip address.
-	 * IPv6 will take precedence.
+	 * Uses the value returned from GetDefaultSocketProtocolFamily to determine default
+	 * addresses to be used.
 	 *
 	 * To skip assumptions, you can call the designated version explicitly below.
 	 */
 	virtual void SetLoopbackAddress() override;
 
+	/**
+	 * Sets the address structure to be bound to the loopback ip address
+	 * for the protocol family specified.
+	 *
+	 * @param ForProtocol the protocol to set the address for
+	 */
+	virtual void SetLoopbackAddress(ESocketProtocolFamily ForProtocol)
+	{
+		if (ForProtocol == ESocketProtocolFamily::IPv4)
+		{
+			SetIPv4LoopbackAddress();
+		}
+		else
+		{
+			SetIPv6LoopbackAddress();
+		}
+	}
+
 	/** Explicit set to loopback IPv4 address */
-	void SetIPv4LoopbackAddress();
+	virtual void SetIPv4LoopbackAddress();
 
 	/** Explicit set to loopback IPv6 address */
-	void SetIPv6LoopbackAddress();
+	virtual void SetIPv6LoopbackAddress();
 
 	/**
 	 * Converts this internet ip address to string form. String will be enclosed in square braces.
diff --git a/Engine/Source/Runtime/Sockets/Private/BSDSockets/SocketSubsystemBSD.h b/Engine/Source/Runtime/Sockets/Private/BSDSockets/SocketSubsystemBSD.h
--- a/Engine/Source/Runtime/Sockets/Private/BSDSockets/SocketSubsystemBSD.h
+++ b/Engine/Source/Runtime/Sockets/Private/BSDSockets/SocketSubsystemBSD.h
@@ -27,16 +27,15 @@ public:
 	 * Specifies the default socket protocol family to use when creating a socket
 	 * without explicitly passing in the protocol type on creation.
 	 *
+	 * This function is mostly here for backwards compatibility. For best practice, moving to
+	 * the new CreateSocket that takes a protocol is advised.
+	 *
 	 * All sockets created using the base class's CreateSocket will use this function
 	 * to determine domain.
 	 */
 	virtual ESocketProtocolFamily GetDefaultSocketProtocolFamily() const
 	{
-#if PLATFORM_HAS_BSD_IPV6_SOCKETS
-		return ESocketProtocolFamily::IPv6;
-#else
 		return ESocketProtocolFamily::IPv4;
-#endif
 	}
 
 	// ISocketSubsystem interface
diff --git a/Engine/Source/Runtime/Sockets/Private/IOS/IPAddressBSDIOS.h b/Engine/Source/Runtime/Sockets/Private/IOS/IPAddressBSDIOS.h
--- a/Engine/Source/Runtime/Sockets/Private/IOS/IPAddressBSDIOS.h
+++ b/Engine/Source/Runtime/Sockets/Private/IOS/IPAddressBSDIOS.h
@@ -14,14 +14,23 @@ class FInternetAddrBSDIOS : public FInternetAddrBSD
 public:
 
 	/** Sets the address to broadcast */
-	virtual void SetBroadcastAddress() override
+	virtual void SetIPv6BroadcastAddress() override
 	{
-		UE_LOG(LogSockets, Verbose, TEXT("SetBroadcastAddress() FInternetAddrBSDIOS"));
-		SetIp(INADDR_BROADCAST);
+		FSocketSubsystemBSD* SocketSubsystem = static_cast<FSocketSubsystemBSD*>(ISocketSubsystem::Get(PLATFORM_SOCKETSUBSYSTEM));
+		if (SocketSubsystem)
+		{
+			// Do a query to get the scope id of the address.
+			if (SocketSubsystem->CreateAddressFromIP("ff02::01", *this) != SE_NO_ERROR)
+			{
+				UE_LOG(LogSockets, Warning, TEXT("Could not resolve the broadcast address for iOS, this address will just be blank"));
+			}
+		}
+		else
+		{
+			UE_LOG(LogSockets, Warning, TEXT("Could not get the socketsubsystem for querying the scope id of the broadcast address"));
+		}
 		SetPort(0);
 	}
-
-	// TODO: Determine if we need to override the others.
 };
 
 #endif
diff --git a/Engine/Source/Runtime/Sockets/Private/IOS/SocketSubsystemIOS.cpp b/Engine/Source/Runtime/Sockets/Private/IOS/SocketSubsystemIOS.cpp
--- a/Engine/Source/Runtime/Sockets/Private/IOS/SocketSubsystemIOS.cpp
+++ b/Engine/Source/Runtime/Sockets/Private/IOS/SocketSubsystemIOS.cpp
@@ -77,6 +77,25 @@ bool FSocketSubsystemIOS::HasNetworkDevice()
 	return true;
 }
 
+ESocketErrors FSocketSubsystemIOS::CreateAddressFromIP(const ANSICHAR* IPAddress, FInternetAddr& OutAddr)
+{
+	return GetHostByName(IPAddress, OutAddr);
+}
+
+ESocketErrors FSocketSubsystemIOS::GetHostByName(const ANSICHAR* HostName, FInternetAddr& OutAddr)
+{
+	FAddressInfoResult GAIResult = GetAddressInfo(ANSI_TO_TCHAR(HostName), nullptr,
+		EAddressInfoFlags::AllResultsWithMapping | EAddressInfoFlags::OnlyUsableAddresses, ESocketProtocolFamily::IPv6);
+
+	if (GAIResult.Results.Num() > 0)
+	{
+		OutAddr.SetRawIp(GAIResult.Results[0].Address->GetRawIp());
+		return SE_NO_ERROR;
+	}
+
+	return SE_HOST_NOT_FOUND;
+}
+
 FSocket* FSocketSubsystemIOS::CreateSocket(const FName& SocketType, const FString& SocketDescription, ESocketProtocolFamily ProtocolType, bool bForceUDP)
 {
 	FSocketBSD* NewSocket = (FSocketBSD*)FSocketSubsystemBSD::CreateSocket(SocketType, SocketDescription, ProtocolType, bForceUDP);
@@ -122,12 +141,14 @@ TSharedRef<FInternetAddr> FSocketSubsystemIOS::GetLocalHostAddr(FOutputDevice& O
 					HostAddr->SetScopeId(ScopeInterfaceId);
 					bWasWifiSet = true;
 					bWasIPv6Set = true;
+					UE_LOG(LogSockets, Verbose, TEXT("Set IP to WIFI %s"), *HostAddr->ToString(false));
 				}
 				else if (!bWasWifiSet && strcmp(Travel->ifa_name, "pdp_ip0") == 0)
 				{
 					HostAddr->SetIp(*AddrData);
 					HostAddr->SetScopeId(ScopeInterfaceId);
 					bWasCellSet = true;
+					UE_LOG(LogSockets, Verbose, TEXT("Set IP to CELL %s"), *HostAddr->ToString(false));
 				}
 			}
 			else if (!bWasIPv6Set && Travel->ifa_addr->sa_family == AF_INET)
@@ -137,12 +158,14 @@ TSharedRef<FInternetAddr> FSocketSubsystemIOS::GetLocalHostAddr(FOutputDevice& O
 					HostAddr->SetIp(*AddrData);
 					HostAddr->SetScopeId(ScopeInterfaceId);
 					bWasWifiSet = true;
+					UE_LOG(LogSockets, Verbose, TEXT("Set IP to WIFI IPv4 %s"), *HostAddr->ToString(false));
 				}
 				else if (!bWasWifiSet && strcmp(Travel->ifa_name, "pdp_ip0") == 0)
 				{
 					HostAddr->SetIp(*AddrData);
 					HostAddr->SetScopeId(ScopeInterfaceId);
 					bWasCellSet = true;
+					UE_LOG(LogSockets, Verbose, TEXT("Set IP to CELL IPv4 %s"), *HostAddr->ToString(false));
 				}
 			}
 		}
diff --git a/Engine/Source/Runtime/Sockets/Private/IOS/SocketSubsystemIOS.h b/Engine/Source/Runtime/Sockets/Private/IOS/SocketSubsystemIOS.h
--- a/Engine/Source/Runtime/Sockets/Private/IOS/SocketSubsystemIOS.h
+++ b/Engine/Source/Runtime/Sockets/Private/IOS/SocketSubsystemIOS.h
@@ -51,6 +51,9 @@ public:
 	{
 	}
 
+	virtual ESocketErrors CreateAddressFromIP(const ANSICHAR* IPAddress, FInternetAddr& OutAddr) override;
+	virtual ESocketErrors GetHostByName(const ANSICHAR* HostName, FInternetAddr& OutAddr) override;
+
 	virtual bool Init(FString& Error) override;
 	virtual void Shutdown() override;
 	virtual bool HasNetworkDevice() override;
