Adding subsystem pointers to the ipaddress to make sure that it can resolve the functions it needs.

diff --git a/Engine/Source/Runtime/Sockets/Private/BSDSockets/IPAddressBSD.cpp b/Engine/Source/Runtime/Sockets/Private/BSDSockets/IPAddressBSD.cpp
--- a/Engine/Source/Runtime/Sockets/Private/BSDSockets/IPAddressBSD.cpp
+++ b/Engine/Source/Runtime/Sockets/Private/BSDSockets/IPAddressBSD.cpp
@@ -6,6 +6,12 @@
 #if PLATFORM_HAS_BSD_SOCKETS || PLATFORM_HAS_BSD_IPV6_SOCKETS
 
 FInternetAddrBSD::FInternetAddrBSD()
+{
+	SocketSubsystem = nullptr;
+	Clear();
+}
+
+FInternetAddrBSD::FInternetAddrBSD(FSocketSubsystemBSD* InSocketSubsystem) : SocketSubsystem(InSocketSubsystem)
 {
 	Clear();
 }
@@ -151,8 +157,7 @@ void FInternetAddrBSD::SetIp(const TCHAR* InAddr, bool& bIsValid)
 	AddressString.RemoveFromEnd("]");
 	
 	const auto InAddrAnsi = StringCast<ANSICHAR>(*AddressString);
-	FSocketSubsystemBSD* SocketSubsystem = static_cast<FSocketSubsystemBSD*>(ISocketSubsystem::Get(PLATFORM_SOCKETSUBSYSTEM));
-	if (SocketSubsystem)
+	if (SocketSubsystem != nullptr)
 	{
 		bIsValid = (SocketSubsystem->CreateAddressFromIP(InAddrAnsi.Get(), *this) == SE_NO_ERROR);
 		if (bHasPort && bIsValid)
@@ -160,6 +165,10 @@ void FInternetAddrBSD::SetIp(const TCHAR* InAddr, bool& bIsValid)
 			SetPort(FCString::Atoi(*Port));
 		}
 	}
+	else
+	{
+		UE_LOG(LogSockets, Verbose, TEXT("SocketSubsystem pointer is null, cannot resolve the stringed address"));
+	}
 }
 
 void FInternetAddrBSD::Set(const sockaddr_storage& AddrData)
@@ -291,8 +300,7 @@ int32 FInternetAddrBSD::GetPort() const
 
 void FInternetAddrBSD::SetAnyAddress()
 {
-	FSocketSubsystemBSD* SocketSubsystem = static_cast<FSocketSubsystemBSD*>(ISocketSubsystem::Get(PLATFORM_SOCKETSUBSYSTEM));
-	if (SocketSubsystem)
+	if (SocketSubsystem != nullptr)
 	{
 		SetAnyAddress(SocketSubsystem->GetDefaultSocketProtocolFamily());
 	}
@@ -319,7 +327,6 @@ void FInternetAddrBSD::SetAnyIPv6Address()
 
 void FInternetAddrBSD::SetBroadcastAddress()
 {
-	FSocketSubsystemBSD* SocketSubsystem = static_cast<FSocketSubsystemBSD*>(ISocketSubsystem::Get(PLATFORM_SOCKETSUBSYSTEM));
 	if (SocketSubsystem)
 	{
 		SetBroadcastAddress(SocketSubsystem->GetDefaultSocketProtocolFamily());
@@ -355,7 +362,6 @@ void FInternetAddrBSD::SetIPv6BroadcastAddress()
 
 void FInternetAddrBSD::SetLoopbackAddress()
 {
-	FSocketSubsystemBSD* SocketSubsystem = static_cast<FSocketSubsystemBSD*>(ISocketSubsystem::Get(PLATFORM_SOCKETSUBSYSTEM));
 	if (SocketSubsystem)
 	{
 		SetLoopbackAddress(SocketSubsystem->GetDefaultSocketProtocolFamily());
@@ -477,7 +483,7 @@ bool FInternetAddrBSD::IsValid() const
 
 TSharedRef<FInternetAddr> FInternetAddrBSD::Clone() const
 {
-	TSharedRef<FInternetAddrBSD> NewAddress = MakeShareable(new FInternetAddrBSD);
+	TSharedRef<FInternetAddrBSD> NewAddress = MakeShareable(new FInternetAddrBSD(SocketSubsystem));
 	NewAddress->Addr = Addr;
 	return NewAddress;
 }
diff --git a/Engine/Source/Runtime/Sockets/Private/BSDSockets/IPAddressBSD.h b/Engine/Source/Runtime/Sockets/Private/BSDSockets/IPAddressBSD.h
--- a/Engine/Source/Runtime/Sockets/Private/BSDSockets/IPAddressBSD.h
+++ b/Engine/Source/Runtime/Sockets/Private/BSDSockets/IPAddressBSD.h
@@ -8,6 +8,8 @@
 
 #if PLATFORM_HAS_BSD_SOCKETS || PLATFORM_HAS_BSD_IPV6_SOCKETS
 
+class FSocketSubsystemBSD;
+
 /**
  * Represents an internet ip address with support for ipv4/v6. 
  * All data is in network byte order
@@ -18,6 +20,9 @@ class FInternetAddrBSD : public FInternetAddr
 	sockaddr_storage Addr;
 
 protected:
+	/** The Subsystem that created this address */
+	FSocketSubsystemBSD* SocketSubsystem;
+
 	/**
 	 * Sets the ip address using a network byte order ipv4 address
 	 *
@@ -51,6 +56,8 @@ public:
 	 */
 	FInternetAddrBSD();
 
+	FInternetAddrBSD(FSocketSubsystemBSD* InSocketSubsystem);
+
 	/**
 	 * Compares FInternetAddrs together, comparing the logical net addresses (endpoints)
 	 * of the data stored, rather than doing a memory comparison like the equality operator does.
diff --git a/Engine/Source/Runtime/Sockets/Private/BSDSockets/SocketSubsystemBSD.cpp b/Engine/Source/Runtime/Sockets/Private/BSDSockets/SocketSubsystemBSD.cpp
--- a/Engine/Source/Runtime/Sockets/Private/BSDSockets/SocketSubsystemBSD.cpp
+++ b/Engine/Source/Runtime/Sockets/Private/BSDSockets/SocketSubsystemBSD.cpp
@@ -141,7 +141,7 @@ FAddressInfoResult FSocketSubsystemBSD::GetAddressInfo(const TCHAR* HostName, co
 				sockaddr_storage* AddrData = reinterpret_cast<sockaddr_storage*>(AddrInfo->ai_addr);
 				if (AddrData != nullptr)
 				{
-					TSharedRef<FInternetAddrBSD> NewAddress = MakeShareable(new FInternetAddrBSD);
+					TSharedRef<FInternetAddrBSD> NewAddress = MakeShareable(new FInternetAddrBSD(this));
 					NewAddress->Set(*AddrData, AddrInfo->ai_addrlen);
 					AddrQueryResult.Results.Add(FAddressInfoResultData(NewAddress, AddrInfo->ai_addrlen,
 						GetProtocolFamilyType(AddrInfo->ai_family), GetSocketType(AddrInfo->ai_protocol)));
@@ -186,7 +186,7 @@ const TCHAR* FSocketSubsystemBSD::GetSocketAPIName() const
 
 TSharedRef<FInternetAddr> FSocketSubsystemBSD::CreateInternetAddr(uint32 Address, uint32 Port)
 {
-	TSharedRef<FInternetAddr> ReturnAddr = MakeShareable(new FInternetAddrBSD);
+	TSharedRef<FInternetAddr> ReturnAddr = MakeShareable(new FInternetAddrBSD(this));
 	ReturnAddr->SetIp(Address);
 	ReturnAddr->SetPort(Port);
 	return ReturnAddr;
diff --git a/Engine/Source/Runtime/Sockets/Private/IOS/IPAddressBSDIOS.h b/Engine/Source/Runtime/Sockets/Private/IOS/IPAddressBSDIOS.h
--- a/Engine/Source/Runtime/Sockets/Private/IOS/IPAddressBSDIOS.h
+++ b/Engine/Source/Runtime/Sockets/Private/IOS/IPAddressBSDIOS.h
@@ -13,6 +13,8 @@ class FInternetAddrBSDIOS : public FInternetAddrBSD
 
 public:
 
+	FInternetAddrBSDIOS(FSocketSubsystemBSD* InSocketSubsystem) : FInternetAddrBSD(InSocketSubsystem) {}
+
 	/** Sets the address to broadcast */
 	virtual void SetIPv6BroadcastAddress() override
 	{
diff --git a/Engine/Source/Runtime/Sockets/Private/IOS/SocketSubsystemIOS.cpp b/Engine/Source/Runtime/Sockets/Private/IOS/SocketSubsystemIOS.cpp
--- a/Engine/Source/Runtime/Sockets/Private/IOS/SocketSubsystemIOS.cpp
+++ b/Engine/Source/Runtime/Sockets/Private/IOS/SocketSubsystemIOS.cpp
@@ -112,7 +112,7 @@ FSocket* FSocketSubsystemIOS::CreateSocket(const FName& SocketType, const FStrin
 
 TSharedRef<FInternetAddr> FSocketSubsystemIOS::GetLocalHostAddr(FOutputDevice& Out, bool& bCanBindAll)
 {
-	TSharedRef<FInternetAddrBSDIOS> HostAddr = MakeShareable(new FInternetAddrBSDIOS);
+	TSharedRef<FInternetAddrBSDIOS> HostAddr = MakeShareable(new FInternetAddrBSDIOS(this));
 	HostAddr->SetAnyAddress();
 
 	ifaddrs* Interfaces = NULL;
@@ -194,7 +194,7 @@ TSharedRef<FInternetAddr> FSocketSubsystemIOS::GetLocalHostAddr(FOutputDevice& O
 
 TSharedRef<FInternetAddr> FSocketSubsystemIOS::CreateInternetAddr(uint32 Address, uint32 Port)
 {
-	TSharedRef<FInternetAddr> Result = MakeShareable(new FInternetAddrBSDIOS);
+	TSharedRef<FInternetAddr> Result = MakeShareable(new FInternetAddrBSDIOS(this));
 	Result->SetIp(Address);
 	Result->SetPort(Port);
 	return Result;
diff --git a/Engine/Source/Runtime/Sockets/Private/Unix/SocketSubsystemUnix.cpp b/Engine/Source/Runtime/Sockets/Private/Unix/SocketSubsystemUnix.cpp
--- a/Engine/Source/Runtime/Sockets/Private/Unix/SocketSubsystemUnix.cpp
+++ b/Engine/Source/Runtime/Sockets/Private/Unix/SocketSubsystemUnix.cpp
@@ -171,7 +171,7 @@ bool FSocketSubsystemUnix::GetLocalAdapterAddresses(TArray<TSharedPtr<FInternetA
 				(Travel->ifa_flags & IFF_LOOPBACK) == 0 && 
 				(AddrFamily == AF_INET || AddrFamily == AF_INET6))
 			{
-				TSharedRef<FInternetAddrBSD> NewAddress = MakeShareable(new FInternetAddrBSD);
+				TSharedRef<FInternetAddrBSD> NewAddress = MakeShareable(new FInternetAddrBSD(this));
 				NewAddress->SetIp(*((sockaddr_storage*)Travel->ifa_addr));
 
 				// Write the scope id if what we found was the multihome address.
diff --git a/Engine/Source/Runtime/Sockets/Private/Windows/SocketSubsystemWindows.cpp b/Engine/Source/Runtime/Sockets/Private/Windows/SocketSubsystemWindows.cpp
--- a/Engine/Source/Runtime/Sockets/Private/Windows/SocketSubsystemWindows.cpp
+++ b/Engine/Source/Runtime/Sockets/Private/Windows/SocketSubsystemWindows.cpp
@@ -168,7 +168,7 @@ bool FSocketSubsystemWindows::GetLocalAdapterAddresses( TArray<TSharedPtr<FInter
 				if ((UnicastAddress->Flags & IP_ADAPTER_ADDRESS_DNS_ELIGIBLE) != 0)
 				{
 					const sockaddr_storage* RawAddress = (const sockaddr_storage*)(UnicastAddress->Address.lpSockaddr);
-					TSharedRef<FInternetAddrBSD> NewAddress = MakeShareable(new FInternetAddrBSD);
+					TSharedRef<FInternetAddrBSD> NewAddress = MakeShareable(new FInternetAddrBSD(this));
 					NewAddress->SetIp(*RawAddress);
 					NewAddress->SetScopeId(ntohl(AdapterAddress->Ipv6IfIndex));
 					OutAdresses.Add(NewAddress);
