Submitting on behalf of Jake.Leonard:
Remove GAI queries on determining the local host address for OSX and Linux. Adding more logging in the getifaddrs queries.

diff --git a/Engine/Source/Runtime/Sockets/Private/Unix/SocketSubsystemUnix.cpp b/Engine/Source/Runtime/Sockets/Private/Unix/SocketSubsystemUnix.cpp
--- a/Engine/Source/Runtime/Sockets/Private/Unix/SocketSubsystemUnix.cpp
+++ b/Engine/Source/Runtime/Sockets/Private/Unix/SocketSubsystemUnix.cpp
@@ -109,21 +109,12 @@ FSocket* FSocketSubsystemUnix::CreateSocket(const FName& SocketType, const FStri
 
 TSharedRef<FInternetAddr> FSocketSubsystemUnix::GetLocalHostAddr(FOutputDevice& Out, bool& bCanBindAll)
 {
-	// get parent address first
-	TSharedRef<FInternetAddr> Addr = FSocketSubsystemBSD::GetLocalHostAddr(Out, bCanBindAll);
-
-	// If the address is not a loopback one (or none), return it.
-	uint32 ParentIp = 0;
-	Addr->GetIp(ParentIp); // will return in host order
-	if (ParentIp != 0 && (ParentIp & 0xff000000) != 0x7f000000)
-	{
-		return Addr;
-	}
+	bCanBindAll = true;
 
 	TArray<TSharedPtr<FInternetAddr>> ResultArray;
 	if (GetLocalAdapterAddresses(ResultArray))
 	{
-		if (FParse::Param(FCommandLine::Get(), TEXT("PRIMARYNET")))
+		if (FParse::Param(FCommandLine::Get(), TEXT("PRIMARYNET")) || FParse::Param(FCommandLine::Get(), TEXT("MULTIHOME")))
 		{
 			bCanBindAll = false;
 		}
@@ -137,6 +128,7 @@ TSharedRef<FInternetAddr> FSocketSubsystemUnix::GetLocalHostAddr(FOutputDevice&
 	}
 
 	// Fall back to this.
+	TSharedRef<FInternetAddr> Addr = CreateInternetAddr();
 	Addr->SetAnyAddress();
 	return Addr;
 }
@@ -154,6 +146,7 @@ bool FSocketSubsystemUnix::GetLocalAdapterAddresses(TArray<TSharedPtr<FInternetA
 
 	ifaddrs* Interfaces = NULL;
 	int InterfaceQueryRet = getifaddrs(&Interfaces);
+	UE_LOG(LogSockets, Verbose, TEXT("Querying net interfaces returned: %d"), InterfaceQueryRet);
 	if (InterfaceQueryRet == 0)
 	{
 		// Loop through linked list of interfaces
@@ -167,24 +160,26 @@ bool FSocketSubsystemUnix::GetLocalAdapterAddresses(TArray<TSharedPtr<FInternetA
 
 			uint16 AddrFamily = Travel->ifa_addr->sa_family;
 			// Find any up and non-loopback addresses
-			if ((Travel->ifa_flags & IFF_UP) != 0 && 
+			if ((Travel->ifa_flags & IFF_UP) != 0 &&
 				(Travel->ifa_flags & IFF_LOOPBACK) == 0 && 
 				(AddrFamily == AF_INET || AddrFamily == AF_INET6))
 			{
 				TSharedRef<FInternetAddrBSD> NewAddress = MakeShareable(new FInternetAddrBSD(this));
 				NewAddress->SetIp(*((sockaddr_storage*)Travel->ifa_addr));
+				uint32 AddressInterface = ntohl(if_nametoindex(Travel->ifa_name));
 
 				// Write the scope id if what we found was the multihome address.
 				// Don't write it to our list again though.
 				if (bHasMultihome && NewAddress == MultihomeAddress)
 				{
-					static_cast<FInternetAddrBSD&>(MultihomeAddress.Get()).SetScopeId(ntohl(if_nametoindex(Travel->ifa_name)));
+					static_cast<FInternetAddrBSD&>(MultihomeAddress.Get()).SetScopeId(AddressInterface);
 				}
 				else
 				{
-					NewAddress->SetScopeId(ntohl(if_nametoindex(Travel->ifa_name)));
+					NewAddress->SetScopeId(AddressInterface);
 					OutAddresses.Add(NewAddress);
 				}
+				UE_LOG(LogSockets, Verbose, TEXT("Got Address %s on interface %d"), *(NewAddress->ToString(false)), AddressInterface);
 			}
 		}
 
