Don't check the bForceUDP flag in SocketSubsystemBSD - this preserves the behavior from 4.20. FTcpSocketBuilder was passing in true for bForceUDP in its CreateSocket call.

diff --git a/Engine/Source/Runtime/Sockets/Private/BSDSockets/SocketSubsystemBSD.cpp b/Engine/Source/Runtime/Sockets/Private/BSDSockets/SocketSubsystemBSD.cpp
--- a/Engine/Source/Runtime/Sockets/Private/BSDSockets/SocketSubsystemBSD.cpp
+++ b/Engine/Source/Runtime/Sockets/Private/BSDSockets/SocketSubsystemBSD.cpp
@@ -44,7 +44,7 @@ FSocket* FSocketSubsystemBSD::CreateSocket(const FName& SocketType, const FStrin
 	PlatformSpecificTypeFlags = SOCK_CLOEXEC;
 #endif // PLATFORM_HAS_BSD_SOCKET_FEATURE_CLOSE_ON_EXEC
 
-	bool bIsUDP = (bForceUDP || SocketType.GetComparisonIndex() == NAME_DGram);
+	bool bIsUDP = SocketType.GetComparisonIndex() == NAME_DGram;
 	int32 SocketTypeFlag = (bIsUDP) ? SOCK_DGRAM : SOCK_STREAM;
 
 	Socket = socket(GetProtocolFamilyValue(ProtocolType), SocketTypeFlag | PlatformSpecificTypeFlags, ((bIsUDP) ? IPPROTO_UDP : IPPROTO_TCP));
