Fix GAI flags to make sure they are valid for the platforms that don't allow certain flags to be explicitly specified.

diff --git a/Engine/Source/Runtime/Sockets/Private/Android/SocketSubsystemAndroid.cpp b/Engine/Source/Runtime/Sockets/Private/Android/SocketSubsystemAndroid.cpp
--- a/Engine/Source/Runtime/Sockets/Private/Android/SocketSubsystemAndroid.cpp
+++ b/Engine/Source/Runtime/Sockets/Private/Android/SocketSubsystemAndroid.cpp
@@ -101,6 +101,33 @@ const TCHAR* FSocketSubsystemAndroid::GetSocketAPIName() const
 	return TEXT("BSD_Android");
 }
 
+ESocketErrors FSocketSubsystemAndroid::GetHostByName(const ANSICHAR* HostName, FInternetAddr& OutAddr)
+{
+	FAddressInfoResult GAIResult = GetAddressInfo(ANSI_TO_TCHAR(HostName), nullptr, EAddressInfoFlags::Default);
+
+	if (GAIResult.Results.Num() > 0)
+	{
+		OutAddr.SetRawIp(GAIResult.Results[0].Address->GetRawIp());
+		return SE_NO_ERROR;
+	}
+
+	return SE_HOST_NOT_FOUND;
+}
+
+ESocketErrors FSocketSubsystemAndroid::CreateAddressFromIP(const ANSICHAR* IPAddress, FInternetAddr& OutAddr)
+{
+	FAddressInfoResult GAIResult = GetAddressInfo(ANSI_TO_TCHAR(IPAddress), nullptr,
+		EAddressInfoFlags::NoResolveHost | EAddressInfoFlags::OnlyUsableAddresses);
+
+	if (GAIResult.Results.Num() > 0)
+	{
+		OutAddr.SetRawIp(GAIResult.Results[0].Address->GetRawIp());
+		return SE_NO_ERROR;
+	}
+
+	return SE_HOST_NOT_FOUND;
+}
+
 
 TSharedRef<FInternetAddr> FSocketSubsystemAndroid::GetLocalHostAddr(FOutputDevice& Out, bool& bCanBindAll)
 {
diff --git a/Engine/Source/Runtime/Sockets/Private/Android/SocketSubsystemAndroid.h b/Engine/Source/Runtime/Sockets/Private/Android/SocketSubsystemAndroid.h
--- a/Engine/Source/Runtime/Sockets/Private/Android/SocketSubsystemAndroid.h
+++ b/Engine/Source/Runtime/Sockets/Private/Android/SocketSubsystemAndroid.h
@@ -70,6 +70,15 @@ public:
 	 */
 	virtual const TCHAR* GetSocketAPIName() const override;
 
+	/**
+	 * Does a DNS look up of a host name
+	 *
+	 * @param HostName the name of the host to look up
+	 * @param Addr the address to copy the IP address to
+	 */
+	virtual ESocketErrors GetHostByName(const ANSICHAR* HostName, FInternetAddr& OutAddr) override;
+	virtual ESocketErrors CreateAddressFromIP(const ANSICHAR* IPAddress, FInternetAddr& OutAddr) override;
+
 	/**
 	 * Android platform specific look up to determine the host address
 	 * as many Android devices have multiple interfaces (wifi, cellular et al.)
