Fix FD_SET macro invoke crash when Socket FD exceed 1023 (contributed by [REDACTED_PERSON])

#PR #4912
#ue4
#android

diff --git a/Engine/Source/Runtime/Sockets/Private/BSDSockets/SocketSubsystemBSD.cpp b/Engine/Source/Runtime/Sockets/Private/BSDSockets/SocketSubsystemBSD.cpp
--- a/Engine/Source/Runtime/Sockets/Private/BSDSockets/SocketSubsystemBSD.cpp
+++ b/Engine/Source/Runtime/Sockets/Private/BSDSockets/SocketSubsystemBSD.cpp
@@ -32,13 +32,45 @@ FSocket* FSocketSubsystemBSD::CreateSocket(const FName& SocketType, const FStrin
 	case NAME_DGram:
 		// Creates a data gram (UDP) socket
 		Socket = socket(AF_INET, SOCK_DGRAM | PlatformSpecificTypeFlags, IPPROTO_UDP);
+
+#if PLATFORM_ANDROID
+		if (Socket != INVALID_SOCKET)
+		{
+			// To avoid out of range in FD_SET
+			if (Socket < 1024)
+			{
+				NewSocket = InternalBSDSocketFactory(Socket, SOCKTYPE_Datagram, SocketDescription);
+			}
+			else
+			{
+				closesocket(Socket);
+			}
+		}
+#else
 		NewSocket = (Socket != INVALID_SOCKET) ? InternalBSDSocketFactory(Socket, SOCKTYPE_Datagram, SocketDescription) : nullptr;
+#endif
 		break;
 
 	case NAME_Stream:
 		// Creates a stream (TCP) socket
 		Socket = socket(AF_INET, SOCK_STREAM | PlatformSpecificTypeFlags, IPPROTO_TCP);
+
+#if PLATFORM_ANDROID
+		if (Socket != INVALID_SOCKET)
+		{
+			// To avoid out of range in FD_SET
+			if (Socket < 1024)
+			{
+				NewSocket = InternalBSDSocketFactory(Socket, SOCKTYPE_Streaming, SocketDescription);
+			}
+			else
+			{
+				closesocket(Socket);
+			}
+		}
+#else
 		NewSocket = (Socket != INVALID_SOCKET) ? InternalBSDSocketFactory(Socket, SOCKTYPE_Streaming, SocketDescription) : nullptr;
+#endif
 		break;
 
 	default:
