Merging //UE4/Dev-Main to Dev-Mobile (//UE4/Dev-Mobile)
Rework PR 4912 from CL 4280211

diff --git a/Engine/Source/Runtime/Sockets/Private/BSDSockets/SocketSubsystemBSD.cpp b/Engine/Source/Runtime/Sockets/Private/BSDSockets/SocketSubsystemBSD.cpp
--- a/Engine/Source/Runtime/Sockets/Private/BSDSockets/SocketSubsystemBSD.cpp
+++ b/Engine/Source/Runtime/Sockets/Private/BSDSockets/SocketSubsystemBSD.cpp
@@ -1,9 +1,8 @@
 // Copyright 1998-2018 Epic Games, Inc. All Rights Reserved.
 
 #include "BSDSockets/SocketSubsystemBSD.h"
-#include "Misc/ScopeLock.h"
 
-#if PLATFORM_HAS_BSD_SOCKETS
+#if PLATFORM_HAS_BSD_SOCKETS || PLATFORM_HAS_BSD_IPV6_SOCKETS
 
 #include "IPAddress.h"
 #include "BSDSockets/IPAddressBSD.h"
@@ -16,65 +15,50 @@ FSocketBSD* FSocketSubsystemBSD::InternalBSDSocketFactory(SOCKET Socket, ESocket
 	return new FSocketBSD(Socket, SocketType, SocketDescription, this);
 }
 
-
 FSocket* FSocketSubsystemBSD::CreateSocket(const FName& SocketType, const FString& SocketDescription, bool bForceUDP)
+{
+	return CreateSocket(SocketType, SocketDescription, GetDefaultSocketProtocolFamily(), bForceUDP);
+}
+
+FSocket* FSocketSubsystemBSD::CreateSocket(const FName& SocketType, const FString& SocketDescription, ESocketProtocolFamily ProtocolType, bool bForceUDP)
 {
 	SOCKET Socket = INVALID_SOCKET;
 	FSocket* NewSocket = nullptr;
 	int PlatformSpecificTypeFlags = 0;
 
+	// For platforms that have two subsystems (ex: Steam) but don't explicitly inherit from SocketSubsystemBSD
+	// so they don't know which protocol to end up using and pass in None.
+	// This is invalid, so we need to attempt to still resolve it.
+	if (ProtocolType == ESocketProtocolFamily::None)
+	{
+		ProtocolType = GetDefaultSocketProtocolFamily();
+	}
+
+	// Don't support any other protocol families.
+	if (ProtocolType != ESocketProtocolFamily::IPv4 && ProtocolType != ESocketProtocolFamily::IPv6)
+	{
+		return nullptr;
+	}
+
 #if PLATFORM_HAS_BSD_SOCKET_FEATURE_CLOSE_ON_EXEC
 	PlatformSpecificTypeFlags = SOCK_CLOEXEC;
 #endif // PLATFORM_HAS_BSD_SOCKET_FEATURE_CLOSE_ON_EXEC
 
-	switch (SocketType.GetComparisonIndex())
-	{
-	case NAME_DGram:
-		// Creates a data gram (UDP) socket
-		Socket = socket(AF_INET, SOCK_DGRAM | PlatformSpecificTypeFlags, IPPROTO_UDP);
-
-#if PLATFORM_ANDROID
-		if (Socket != INVALID_SOCKET)
-		{
-			// To avoid out of range in FD_SET
-			if (Socket < 1024)
-			{
-				NewSocket = InternalBSDSocketFactory(Socket, SOCKTYPE_Datagram, SocketDescription);
-			}
-			else
-			{
-				closesocket(Socket);
-			}
-		}
-#else
-		NewSocket = (Socket != INVALID_SOCKET) ? InternalBSDSocketFactory(Socket, SOCKTYPE_Datagram, SocketDescription) : nullptr;
-#endif
-		break;
+	bool bIsUDP = SocketType.GetComparisonIndex() == NAME_DGram;
+	int32 SocketTypeFlag = (bIsUDP) ? SOCK_DGRAM : SOCK_STREAM;
 
-	case NAME_Stream:
-		// Creates a stream (TCP) socket
-		Socket = socket(AF_INET, SOCK_STREAM | PlatformSpecificTypeFlags, IPPROTO_TCP);
+	Socket = socket(GetProtocolFamilyValue(ProtocolType), SocketTypeFlag | PlatformSpecificTypeFlags, ((bIsUDP) ? IPPROTO_UDP : IPPROTO_TCP));
 
 #if PLATFORM_ANDROID
-		if (Socket != INVALID_SOCKET)
-		{
-			// To avoid out of range in FD_SET
-			if (Socket < 1024)
-			{
-				NewSocket = InternalBSDSocketFactory(Socket, SOCKTYPE_Streaming, SocketDescription);
-			}
-			else
-			{
-				closesocket(Socket);
-			}
-		}
-#else
-		NewSocket = (Socket != INVALID_SOCKET) ? InternalBSDSocketFactory(Socket, SOCKTYPE_Streaming, SocketDescription) : nullptr;
+	// To avoid out of range in FD_SET
+	if (Socket != INVALID_SOCKET && Socket >= 1024)
+	{
+		closesocket(Socket);
+	}
+	else
 #endif
-		break;
-
-	default:
-		break;
+	{
+		NewSocket = (Socket != INVALID_SOCKET) ? InternalBSDSocketFactory(Socket, ((bIsUDP) ? SOCKTYPE_Datagram : SOCKTYPE_Streaming), SocketDescription) : nullptr;
 	}
 
 	if (!NewSocket)
@@ -85,53 +69,110 @@ FSocket* FSocketSubsystemBSD::CreateSocket(const FName& SocketType, const FStrin
 	return NewSocket;
 }
 
+ESocketErrors FSocketSubsystemBSD::CreateAddressFromIP(const ANSICHAR* IPAddress, FInternetAddr& OutAddr)
+{
+	FAddressInfoResult GAIResult = GetAddressInfo(ANSI_TO_TCHAR(IPAddress), nullptr,
+		EAddressInfoFlags::AllResultsWithMapping | EAddressInfoFlags::NoResolveHost | EAddressInfoFlags::OnlyUsableAddresses);
+
+	if (GAIResult.Results.Num() > 0)
+	{
+		OutAddr.SetRawIp(GAIResult.Results[0].Address->GetRawIp());
+		return SE_NO_ERROR;
+	}
+	
+	return SE_HOST_NOT_FOUND;
+}
+
+ESocketErrors FSocketSubsystemBSD::GetHostByName(const ANSICHAR* HostName, FInternetAddr& OutAddr)
+{
+	FAddressInfoResult GAIResult = GetAddressInfo(ANSI_TO_TCHAR(HostName), nullptr,
+		EAddressInfoFlags::AllResultsWithMapping | EAddressInfoFlags::OnlyUsableAddresses | EAddressInfoFlags::BindableAddress);
+
+	if (GAIResult.Results.Num() > 0)
+	{
+		OutAddr.SetRawIp(GAIResult.Results[0].Address->GetRawIp());
+		return SE_NO_ERROR;
+	}
+
+	return SE_HOST_NOT_FOUND;
+}
 
 void FSocketSubsystemBSD::DestroySocket(FSocket* Socket)
 {
 	delete Socket;
 }
 
-
-ESocketErrors FSocketSubsystemBSD::GetHostByName(const ANSICHAR* HostName, FInternetAddr& OutAddr)
+FAddressInfoResult FSocketSubsystemBSD::GetAddressInfo(const TCHAR* HostName, const TCHAR* ServiceName,
+	EAddressInfoFlags QueryFlags, ESocketProtocolFamily ProtocolType, ESocketType SocketType)
 {
+	FAddressInfoResult AddrQueryResult = FAddressInfoResult(HostName, ServiceName);
+
+	if (HostName == nullptr && ServiceName == nullptr)
+	{
+		UE_LOG(LogSockets, Warning, TEXT("GetAddressInfo was passed with both a null host and service, returning empty result"));
+		return AddrQueryResult;
+	}
+
 #if PLATFORM_HAS_BSD_SOCKET_FEATURE_GETADDRINFO
-	FScopeLock ScopeLock(&HostByNameSynch);
 	addrinfo* AddrInfo = nullptr;
 
-	// Limit the IP Addresses we get back to just IPv4
+	// Make sure we filter out IPv6 if the platform is not officially supported
+	// (if it isn't supported but we explicitly ask for it, allow it).
+	bool bCanUseIPv6 = (PLATFORM_HAS_BSD_IPV6_SOCKETS || ProtocolType == ESocketProtocolFamily::IPv6) ? true : false;
+
 	addrinfo HintAddrInfo;
 	FMemory::Memzero(&HintAddrInfo, sizeof(HintAddrInfo));
-	HintAddrInfo.ai_family = AF_INET;
+	HintAddrInfo.ai_family = GetProtocolFamilyValue(ProtocolType);
+	HintAddrInfo.ai_flags = GetAddressInfoHintFlag(QueryFlags);
+
+	if (SocketType != ESocketType::SOCKTYPE_Unknown)
+	{
+		bool bIsUDP = (SocketType == ESocketType::SOCKTYPE_Datagram);
+		HintAddrInfo.ai_protocol = bIsUDP ? IPPROTO_UDP : IPPROTO_TCP;
+		HintAddrInfo.ai_socktype = bIsUDP ? SOCK_DGRAM : SOCK_STREAM;
+	}
 
-	int32 ErrorCode = getaddrinfo(HostName, nullptr, &HintAddrInfo, &AddrInfo);
+	int32 ErrorCode = getaddrinfo(TCHAR_TO_UTF8(HostName), TCHAR_TO_UTF8(ServiceName), &HintAddrInfo, &AddrInfo);
 	ESocketErrors SocketError = TranslateGAIErrorCode(ErrorCode);
+
+	UE_LOG(LogSockets, Verbose, TEXT("Executed getaddrinfo with HostName: %s Return: %d"), HostName, ErrorCode);
 	if (SocketError == SE_NO_ERROR)
 	{
+		addrinfo* AddrInfoHead = AddrInfo;
+		// The canonical name will always be stored in only the first result in a getaddrinfo query
+		if (AddrInfo != nullptr && AddrInfo->ai_canonname != nullptr)
+		{
+			AddrQueryResult.CanonicalNameResult = UTF8_TO_TCHAR(AddrInfo->ai_canonname);
+		}
+
 		for (; AddrInfo != nullptr; AddrInfo = AddrInfo->ai_next)
 		{
-			if (AddrInfo->ai_family == AF_INET)
+			if (AddrInfo->ai_family == AF_INET || (AddrInfo->ai_family == AF_INET6 && bCanUseIPv6))
 			{
-				sockaddr_in* IPv4SockAddr = reinterpret_cast<sockaddr_in*>(AddrInfo->ai_addr);
-				if (IPv4SockAddr != nullptr)
+				sockaddr_storage* AddrData = reinterpret_cast<sockaddr_storage*>(AddrInfo->ai_addr);
+				if (AddrData != nullptr)
 				{
-					uint32 HostIP = ntohl(IPv4SockAddr->sin_addr.s_addr);
-					static_cast<FInternetAddrBSD&>(OutAddr).SetIp(HostIP);
-					freeaddrinfo(AddrInfo);
-					return SE_NO_ERROR;
+					TSharedRef<FInternetAddrBSD> NewAddress = MakeShareable(new FInternetAddrBSD(this));
+					NewAddress->Set(*AddrData, AddrInfo->ai_addrlen);
+					AddrQueryResult.Results.Add(FAddressInfoResultData(NewAddress, AddrInfo->ai_addrlen,
+						GetProtocolFamilyType(AddrInfo->ai_family), GetSocketType(AddrInfo->ai_protocol)));
+
+					UE_LOG(LogSockets, Verbose, TEXT("# Family: %s Address: %s"), ((AddrInfo->ai_family == AF_INET) ? TEXT("IPv4") : TEXT("IPv6")), *(NewAddress->ToString(false)));
 				}
 			}
 		}
-		freeaddrinfo(AddrInfo);
-		return SE_HOST_NOT_FOUND;
+		freeaddrinfo(AddrInfoHead);
+	}
+	else
+	{
+		UE_LOG(LogSockets, Warning, TEXT("GetAddressInfo failed to resolve host with error %s [%d]"), GetSocketError(SocketError), ErrorCode);
 	}
-	return SocketError;
 #else
-	UE_LOG(LogSockets, Error, TEXT("Platform has no getaddrinfo(), but did not override FSocketSubsystem::GetHostByName()"));
-	return SE_NO_RECOVERY;
+	UE_LOG(LogSockets, Error, TEXT("Platform has no getaddrinfo(), but did not override FSocketSubsystem::GetAddressInfo()"));
 #endif
+	return AddrQueryResult;
 }
 
-
 bool FSocketSubsystemBSD::GetHostName(FString& HostName)
 {
 #if PLATFORM_HAS_BSD_SOCKET_FEATURE_GETHOSTNAME
@@ -151,19 +192,17 @@ bool FSocketSubsystemBSD::GetHostName(FString& HostName)
 
 const TCHAR* FSocketSubsystemBSD::GetSocketAPIName() const
 {
-	return TEXT("BSD");
+	return TEXT("BSD IPv4/6");
 }
 
-
 TSharedRef<FInternetAddr> FSocketSubsystemBSD::CreateInternetAddr(uint32 Address, uint32 Port)
 {
-	TSharedRef<FInternetAddr> Result = MakeShareable(new FInternetAddrBSD);
-	Result->SetIp(Address);
-	Result->SetPort(Port);
-	return Result;
+	TSharedRef<FInternetAddr> ReturnAddr = MakeShareable(new FInternetAddrBSD(this));
+	ReturnAddr->SetIp(Address);
+	ReturnAddr->SetPort(Port);
+	return ReturnAddr;
 }
 
-
 ESocketErrors FSocketSubsystemBSD::GetLastErrorCode()
 {
 	return TranslateErrorCode(errno);
@@ -251,5 +290,129 @@ ESocketErrors FSocketSubsystemBSD::TranslateErrorCode(int32 Code)
 	return SE_EINVAL;
 }
 
+ESocketErrors FSocketSubsystemBSD::TranslateGAIErrorCode(int32 Code) const
+{
+#if PLATFORM_HAS_BSD_SOCKET_FEATURE_GETADDRINFO
+	switch (Code)
+	{
+	// getaddrinfo() has its own error codes
+	case EAI_AGAIN:			return SE_TRY_AGAIN;
+	case EAI_BADFLAGS:		return SE_EINVAL;
+	case EAI_FAIL:			return SE_NO_RECOVERY;
+	case EAI_FAMILY:		return SE_EAFNOSUPPORT;
+	case EAI_MEMORY:		return SE_ENOBUFS;
+	case EAI_NONAME:		return SE_HOST_NOT_FOUND;
+	case EAI_SERVICE:		return SE_EPFNOSUPPORT;
+	case EAI_SOCKTYPE:		return SE_ESOCKTNOSUPPORT;
+#if PLATFORM_HAS_BSD_SOCKET_FEATURE_WINSOCKETS
+	case WSANO_DATA:		return SE_NO_DATA;
+	case WSANOTINITIALISED: return SE_NOTINITIALISED;
+#else			
+	case EAI_NODATA:		return SE_NO_DATA;
+	case EAI_ADDRFAMILY:	return SE_ADDRFAMILY;
+	case EAI_SYSTEM:		return SE_SYSTEM;
+#endif
+	case 0:					break; // 0 means success
+	default:
+		UE_LOG(LogSockets, Warning, TEXT("Unhandled getaddrinfo() socket error! Code: %d"), Code);
+		return SE_EINVAL;
+	}
+#endif // PLATFORM_HAS_BSD_SOCKET_FEATURE_GETADDRINFO
+
+	return SE_NO_ERROR;
+}
+
+int32 FSocketSubsystemBSD::GetProtocolFamilyValue(ESocketProtocolFamily InProtocol) const
+{
+	switch (InProtocol)
+	{
+		default:
+		case ESocketProtocolFamily::None: return AF_UNSPEC;
+		case ESocketProtocolFamily::IPv4: return AF_INET;
+		case ESocketProtocolFamily::IPv6: return AF_INET6;
+	}
+}
+
+ESocketProtocolFamily FSocketSubsystemBSD::GetProtocolFamilyType(int32 InProtocol) const
+{
+	switch (InProtocol)
+	{
+		default:
+		case AF_UNSPEC: return ESocketProtocolFamily::None;
+		case AF_INET:   return ESocketProtocolFamily::IPv4;
+		case AF_INET6:  return ESocketProtocolFamily::IPv6;
+	}
+}
+
+ESocketType FSocketSubsystemBSD::GetSocketType(int32 InSocketType) const
+{
+	switch (InSocketType)
+	{
+		case SOCK_STREAM:
+		case IPPROTO_TCP: return ESocketType::SOCKTYPE_Streaming;
+		case SOCK_DGRAM:
+		case IPPROTO_UDP: return ESocketType::SOCKTYPE_Datagram;
+		default: return SOCKTYPE_Unknown;
+	}
+}
+
+int32 FSocketSubsystemBSD::GetAddressInfoHintFlag(EAddressInfoFlags InFlags) const
+{
+	int32 ReturnFlagsCode = 0;
+
+#if PLATFORM_HAS_BSD_SOCKET_FEATURE_GETADDRINFO
+	if (InFlags == EAddressInfoFlags::Default)
+	{
+		return ReturnFlagsCode;
+	}
+
+	if (EnumHasAnyFlags(InFlags, EAddressInfoFlags::NoResolveHost))
+	{
+		ReturnFlagsCode |= AI_NUMERICHOST;
+	}
+
+	if (EnumHasAnyFlags(InFlags, EAddressInfoFlags::NoResolveService))
+	{
+		ReturnFlagsCode |= AI_NUMERICSERV;
+	}
+
+	if (EnumHasAnyFlags(InFlags, EAddressInfoFlags::OnlyUsableAddresses))
+	{
+		ReturnFlagsCode |= AI_ADDRCONFIG;
+	}
+
+	if (EnumHasAnyFlags(InFlags, EAddressInfoFlags::BindableAddress))
+	{
+		ReturnFlagsCode |= AI_PASSIVE;
+	}
+
+	/* This means nothing unless AI_ALL is also specified. */
+	if (EnumHasAnyFlags(InFlags, EAddressInfoFlags::AllowV4MappedAddresses))
+	{
+		ReturnFlagsCode |= AI_V4MAPPED;
+	}
+
+	if (EnumHasAnyFlags(InFlags, EAddressInfoFlags::AllResults))
+	{
+		ReturnFlagsCode |= AI_ALL;
+	}
+
+	if (EnumHasAnyFlags(InFlags, EAddressInfoFlags::CanonicalName))
+	{
+		ReturnFlagsCode |= AI_CANONNAME;
+	}
+
+	if (EnumHasAnyFlags(InFlags, EAddressInfoFlags::FQDomainName))
+	{
+#ifdef AI_FQDN
+		ReturnFlagsCode |= AI_FQDN;
+#else
+		ReturnFlagsCode |= AI_CANONNAME;
+#endif
+	}
+#endif
+
+	return ReturnFlagsCode;
+}
 
 #endif	//PLATFORM_HAS_BSD_SOCKETS
