Updated log verbosity of getaddrinfo results so they will log by default, and added port information.

diff --git a/Engine/Source/Runtime/Sockets/Private/BSDSockets/SocketSubsystemBSD.cpp b/Engine/Source/Runtime/Sockets/Private/BSDSockets/SocketSubsystemBSD.cpp
--- a/Engine/Source/Runtime/Sockets/Private/BSDSockets/SocketSubsystemBSD.cpp
+++ b/Engine/Source/Runtime/Sockets/Private/BSDSockets/SocketSubsystemBSD.cpp
@@ -137,7 +137,7 @@ FAddressInfoResult FSocketSubsystemBSD::GetAddressInfo(const TCHAR* HostName, co
 	int32 ErrorCode = getaddrinfo(TCHAR_TO_UTF8(HostName), TCHAR_TO_UTF8(ServiceName), &HintAddrInfo, &AddrInfo);
 	ESocketErrors SocketError = TranslateGAIErrorCode(ErrorCode);
 
-	UE_LOG(LogSockets, Verbose, TEXT("Executed getaddrinfo with HostName: %s Return: %d"), HostName, ErrorCode);
+	UE_LOG(LogSockets, Log, TEXT("Executed getaddrinfo with HostName: %s ServiceName: %s Return: %d"), HostName ? HostName : TEXT("null"), ServiceName ? ServiceName : TEXT("null"), ErrorCode);
 	if (SocketError == SE_NO_ERROR)
 	{
 		addrinfo* AddrInfoHead = AddrInfo;
@@ -159,7 +159,7 @@ FAddressInfoResult FSocketSubsystemBSD::GetAddressInfo(const TCHAR* HostName, co
 					AddrQueryResult.Results.Add(FAddressInfoResultData(NewAddress, AddrInfo->ai_addrlen,
 						GetProtocolFamilyType(AddrInfo->ai_family), GetSocketType(AddrInfo->ai_protocol)));
 
-					UE_LOG(LogSockets, Verbose, TEXT("# Family: %s Address: %s"), ((AddrInfo->ai_family == AF_INET) ? TEXT("IPv4") : TEXT("IPv6")), *(NewAddress->ToString(false)));
+					UE_LOG(LogSockets, Log, TEXT("# Family: %s Address: %s"), ((AddrInfo->ai_family == AF_INET) ? TEXT("IPv4") : TEXT("IPv6")), *(NewAddress->ToString(true)));
 				}
 			}
 		}
