Because addrconfig is not supported on HTML5. This is not documented anywhere, so any query with the flag will just be removed from any queries.

diff --git a/Engine/Source/Runtime/Sockets/Private/HTML5/Device/SocketSubsystemHTML5.cpp b/Engine/Source/Runtime/Sockets/Private/HTML5/Device/SocketSubsystemHTML5.cpp
--- a/Engine/Source/Runtime/Sockets/Private/HTML5/Device/SocketSubsystemHTML5.cpp
+++ b/Engine/Source/Runtime/Sockets/Private/HTML5/Device/SocketSubsystemHTML5.cpp
@@ -80,3 +80,15 @@ bool FSocketSubsystemHTML5::HasNetworkDevice()
 {
 	return true;
 }
+
+/**
+ * Translates an ESocketAddressInfoFlags into a value usable by getaddrinfo
+ */
+int32 FSocketSubsystemHTML5::GetAddressInfoHintFlag(EAddressInfoFlags InFlags) const
+{
+	// As of writing, emscripten does not support AI_ADDRCONFIG. It is marked as an usable flag,
+	// however if it is set, GAI will fail out with a bad name flag.
+	// As such, remove the flag from any potential queries.
+	EAddressInfoFlags ModifiedInFlags = (InFlags & ~EAddressInfoFlags::OnlyUsableAddresses);
+	return FSocketSubsystemBSD::GetAddressInfoHintFlag(ModifiedInFlags);
+}
diff --git a/Engine/Source/Runtime/Sockets/Private/HTML5/Device/SocketSubsystemHTML5.h b/Engine/Source/Runtime/Sockets/Private/HTML5/Device/SocketSubsystemHTML5.h
--- a/Engine/Source/Runtime/Sockets/Private/HTML5/Device/SocketSubsystemHTML5.h
+++ b/Engine/Source/Runtime/Sockets/Private/HTML5/Device/SocketSubsystemHTML5.h
@@ -60,4 +60,9 @@ public:
 	 * @return Whether the device has a properly configured network device or not
 	 */
 	virtual bool HasNetworkDevice() override;
+
+	/**
+	 * Translates an ESocketAddressInfoFlags into a value usable by getaddrinfo
+	 */
+	virtual int32 GetAddressInfoHintFlag(EAddressInfoFlags InFlags) const override;
 };
