Fix using SetIp(uint) on an IPv6 only platform so that both cases are handled properly. Should contain the same flow as the IPv6 framework.

diff --git a/Engine/Source/Runtime/Sockets/Private/BSDSockets/IPAddressBSD.cpp b/Engine/Source/Runtime/Sockets/Private/BSDSockets/IPAddressBSD.cpp
--- a/Engine/Source/Runtime/Sockets/Private/BSDSockets/IPAddressBSD.cpp
+++ b/Engine/Source/Runtime/Sockets/Private/BSDSockets/IPAddressBSD.cpp
@@ -5,6 +5,22 @@
 
 #if PLATFORM_HAS_BSD_SOCKETS || PLATFORM_HAS_BSD_IPV6_SOCKETS
 
+// Hardcoded address that the messagebus uses. This is a hack.
+#define IPV4_MESSAGEBUS_ADDRESS_HACK ((230 << 24) | (0 << 16) | (0 << 8) | (1 << 0))
+
+#if PLATFORM_HAS_BSD_IPV6_SOCKETS
+void MapIPv4ToIPv6(const uint32& InAddress, in6_addr& OutStructure)
+{
+	FMemory::Memzero(OutStructure);
+	OutStructure.s6_addr[10] = 0xff;
+	OutStructure.s6_addr[11] = 0xff;
+	OutStructure.s6_addr[12] = (static_cast<uint32>(InAddress) & 0xFF);
+	OutStructure.s6_addr[13] = ((static_cast<uint32>(InAddress) >> 8) & 0xFF);
+	OutStructure.s6_addr[14] = ((static_cast<uint32>(InAddress) >> 16) & 0xFF);
+	OutStructure.s6_addr[15] = ((static_cast<uint32>(InAddress) >> 24) & 0xFF);
+}
+#endif
+
 FInternetAddrBSD::FInternetAddrBSD()
 {
 	SocketSubsystem = nullptr;
@@ -56,14 +72,7 @@ bool FInternetAddrBSD::CompareEndpoints(const FInternetAddr& InAddr) const
 		}
 
 		in6_addr ConvertedAddrData;
-		FMemory::Memzero(ConvertedAddrData);
-		ConvertedAddrData.s6_addr[10] = 0xff;
-		ConvertedAddrData.s6_addr[11] = 0xff;
-		ConvertedAddrData.s6_addr[12] = (static_cast<uint32>(IPv4Addr->s_addr) & 0xFF);
-		ConvertedAddrData.s6_addr[13] = ((static_cast<uint32>(IPv4Addr->s_addr) >> 8) & 0xFF);
-		ConvertedAddrData.s6_addr[14] = ((static_cast<uint32>(IPv4Addr->s_addr) >> 16) & 0xFF);
-		ConvertedAddrData.s6_addr[15] = ((static_cast<uint32>(IPv4Addr->s_addr) >> 24) & 0xFF);
-
+		MapIPv4ToIPv6(IPv4Addr->s_addr, ConvertedAddrData);
 		return memcmp(&(ConvertedAddrData), IPv6Addr, sizeof(in6_addr)) == 0;
 #else
 		return false;
@@ -171,6 +180,33 @@ void FInternetAddrBSD::SetIp(const TCHAR* InAddr, bool& bIsValid)
 	}
 }
 
+void FInternetAddrBSD::SetIp(uint32 InAddr)
+{
+#if PLATFORM_HAS_BSD_IPV6_SOCKETS
+	if (SocketSubsystem && SocketSubsystem->GetDefaultSocketProtocolFamily() == ESocketProtocolFamily::IPv6)
+	{
+		if (InAddr == 0)
+		{
+			SetAnyIPv6Address();
+		}
+		else if (InAddr == INADDR_BROADCAST || InAddr == IPV4_MESSAGEBUS_ADDRESS_HACK)
+		{
+			SetIPv6BroadcastAddress();
+		}
+		else
+		{
+			in6_addr ConvertedAddrData;
+			MapIPv4ToIPv6(InAddr, ConvertedAddrData);
+			SetIp(ConvertedAddrData);
+		}
+		return;
+	}
+#endif
+
+	((sockaddr_in*)&Addr)->sin_addr.s_addr = htonl(InAddr);
+	Addr.ss_family = AF_INET;
+}
+
 void FInternetAddrBSD::Set(const sockaddr_storage& AddrData)
 {
 	Addr = AddrData;
diff --git a/Engine/Source/Runtime/Sockets/Private/BSDSockets/IPAddressBSD.h b/Engine/Source/Runtime/Sockets/Private/BSDSockets/IPAddressBSD.h
--- a/Engine/Source/Runtime/Sockets/Private/BSDSockets/IPAddressBSD.h
+++ b/Engine/Source/Runtime/Sockets/Private/BSDSockets/IPAddressBSD.h
@@ -83,15 +83,12 @@ public:
 	virtual TArray<uint8> GetRawIp() const override;
 
 	/**
-	 * Sets the ip address from a host byte order uint32
+	 * Sets the ip address from a host byte order uint32.
+	 * Will attempt to determine platform compatibility for best way to store address
 	 *
 	 * @param InAddr the new address to use (must convert to network byte order)
 	 */
-	virtual void SetIp(uint32 InAddr) override
-	{
-		((sockaddr_in*)&Addr)->sin_addr.s_addr = htonl(InAddr);
-		Addr.ss_family = AF_INET;
-	}
+	virtual void SetIp(uint32 InAddr) override;
 
 	/**
 	 * Sets the ip address from a string IPv6 or IPv4 address.
@@ -163,7 +160,7 @@ public:
 
 	/**
 	 * Copies the network byte order ip address to a host byte order dword.
-	 * Does nothing if we are currently not storing an ipv4 addr
+	 * Does nothing if we are currently not storing an ipv4 addr in some capacity
 	 *
 	 * @param OutAddr the out param receiving the ip address
 	 */
