Fix for sessions not appearing in the Session Frontend after launching on IOS

diff --git a/Engine/Source/Runtime/Sockets/Private/BSDSockets/SocketsBSD.cpp b/Engine/Source/Runtime/Sockets/Private/BSDSockets/SocketsBSD.cpp
--- a/Engine/Source/Runtime/Sockets/Private/BSDSockets/SocketsBSD.cpp
+++ b/Engine/Source/Runtime/Sockets/Private/BSDSockets/SocketsBSD.cpp
@@ -423,8 +423,7 @@ bool FSocketBSD::LeaveMulticastGroup(const FInternetAddr& GroupAddress)
 bool FSocketBSD::SetMulticastLoopback(bool bLoopback)
 {
 #if PLATFORM_HAS_BSD_IPV6_SOCKETS
-	bool IPv6Loop = bLoopback;
-	setsockopt(Socket, IPPROTO_IPV6, IPV6_MULTICAST_LOOP, (char*)&IPv6Loop, sizeof(IPv6Loop));
+	return (setsockopt(Socket, IPPROTO_IPV6, IP_MULTICAST_LOOP, (char*)&bLoopback, sizeof(bLoopback)) == 0);
 #endif
 
 	return (setsockopt(Socket, IPPROTO_IP, IP_MULTICAST_LOOP, (char*)&bLoopback, sizeof(bLoopback)) == 0);
@@ -434,8 +433,7 @@ bool FSocketBSD::SetMulticastLoopback(bool bLoopback)
 bool FSocketBSD::SetMulticastTtl(uint8 TimeToLive)
 {
 #if PLATFORM_HAS_BSD_IPV6_SOCKETS
-	uint8 IPv6TTL = TimeToLive;
-	setsockopt(Socket, IPPROTO_IPV6, IP_MULTICAST_TTL, (char*)&IPv6TTL, sizeof(IPv6TTL));
+	return (setsockopt(Socket, IPPROTO_IPV6, IP_MULTICAST_TTL, (char*)&TimeToLive, sizeof(TimeToLive)) == 0);
 #endif
 
 	return (setsockopt(Socket, IPPROTO_IP, IP_MULTICAST_TTL, (char*)&TimeToLive, sizeof(TimeToLive)) == 0);
