Allow sockets to know what their protocols are. This will be handy for debugging the helper functions when they start getting added. This also helps because there's no platform agnostic way (atleast to my knowledge, I know Winsock has one but it's their own special struct) to pull what protocol a socket supports at runtime after creation.

Also fix up some missing protocol defaults.

diff --git a/Engine/Plugins/Online/OnlineSubsystemSteam/Source/Private/SocketSubsystemSteam.cpp b/Engine/Plugins/Online/OnlineSubsystemSteam/Source/Private/SocketSubsystemSteam.cpp
--- a/Engine/Plugins/Online/OnlineSubsystemSteam/Source/Private/SocketSubsystemSteam.cpp
+++ b/Engine/Plugins/Online/OnlineSubsystemSteam/Source/Private/SocketSubsystemSteam.cpp
@@ -163,7 +163,7 @@ FSocket* FSocketSubsystemSteam::CreateSocket(const FName& SocketType, const FStr
 		if (SteamUserPtr != NULL)
 		{
 			FUniqueNetIdSteam ClientId(SteamUserPtr->GetSteamID());
-			NewSocket = new FSocketSteam(SteamNetworking(), ClientId, SocketDescription);
+			NewSocket = new FSocketSteam(SteamNetworking(), ClientId, SocketDescription, ProtocolType);
 
 			if (NewSocket)
 			{
@@ -180,12 +180,12 @@ FSocket* FSocketSubsystemSteam::CreateSocket(const FName& SocketType, const FStr
 			// If the GameServer connection hasn't been created yet, mark the socket as invalid for now
 			if (SessionInt->bSteamworksGameServerConnected && SessionInt->GameServerSteamId->IsValid() && SessionInt->bPolicyResponseReceived)
 			{
-				NewSocket = new FSocketSteam(SteamGameServerNetworking(), *SessionInt->GameServerSteamId, SocketDescription);
+				NewSocket = new FSocketSteam(SteamGameServerNetworking(), *SessionInt->GameServerSteamId, SocketDescription, ProtocolType);
 			}
 			else
 			{
 				FUniqueNetIdSteam InvalidId(uint64(0));
-				NewSocket = new FSocketSteam(SteamGameServerNetworking(), InvalidId, SocketDescription);
+				NewSocket = new FSocketSteam(SteamGameServerNetworking(), InvalidId, SocketDescription, ProtocolType);
 			}
 
 			if (NewSocket)
diff --git a/Engine/Plugins/Online/OnlineSubsystemSteam/Source/Private/SocketsSteam.h b/Engine/Plugins/Online/OnlineSubsystemSteam/Source/Private/SocketsSteam.h
--- a/Engine/Plugins/Online/OnlineSubsystemSteam/Source/Private/SocketsSteam.h
+++ b/Engine/Plugins/Online/OnlineSubsystemSteam/Source/Private/SocketsSteam.h
@@ -48,14 +48,15 @@ PACKAGE_SCOPE:
 
 public:
 	/**
-	 * Assigns a Windows socket to this object
+	 * Creates a Steam socket
 	 *
 	 * @param InSocket the socket to assign to this object
 	 * @param InSocketType the type of socket that was created
 	 * @param InSocketDescription the debug description of the socket
+	 * @param InSocketProtocol the protocol used to create this socket.
 	 */
-	FSocketSteam(ISteamNetworking* InSteamNetworkingPtr, FUniqueNetIdSteam& InLocalSteamId, const FString& InSocketDescription) :
-		FSocket(SOCKTYPE_Datagram, InSocketDescription),
+	FSocketSteam(ISteamNetworking* InSteamNetworkingPtr, FUniqueNetIdSteam& InLocalSteamId, const FString& InSocketDescription, ESocketProtocolFamily InSocketProtocol) :
+		FSocket(SOCKTYPE_Datagram, InSocketDescription, InSocketProtocol),
 		LocalSteamId(InLocalSteamId),
 		SteamChannel(0),
 		SteamSendMode(k_EP2PSendUnreliable),
diff --git a/Engine/Plugins/Online/OnlineSubsystemUtils/Source/OnlineSubsystemUtils/Private/IpNetDriver.cpp b/Engine/Plugins/Online/OnlineSubsystemUtils/Source/OnlineSubsystemUtils/Private/IpNetDriver.cpp
--- a/Engine/Plugins/Online/OnlineSubsystemUtils/Source/OnlineSubsystemUtils/Private/IpNetDriver.cpp
+++ b/Engine/Plugins/Online/OnlineSubsystemUtils/Source/OnlineSubsystemUtils/Private/IpNetDriver.cpp
@@ -119,7 +119,7 @@ bool UIpNetDriver::InitBase( bool bInitAsClient, FNetworkNotify* InNotify, const
 	if( Socket == NULL )
 	{
 		Socket = 0;
-		Error = FString::Printf( TEXT("WinSock: socket failed (%i)"), (int32)SocketSubsystem->GetLastErrorCode() );
+		Error = FString::Printf( TEXT("%s: socket failed (%i)"), SocketSubsystem->GetSocketAPIName(), (int32)SocketSubsystem->GetLastErrorCode() );
 		return false;
 	}
 	if (SocketSubsystem->RequiresChatDataBeSeparate() == false &&
diff --git a/Engine/Source/Runtime/Sockets/Private/BSDIPv6Sockets/SocketSubsystemBSDIPv6.cpp b/Engine/Source/Runtime/Sockets/Private/BSDIPv6Sockets/SocketSubsystemBSDIPv6.cpp
--- a/Engine/Source/Runtime/Sockets/Private/BSDIPv6Sockets/SocketSubsystemBSDIPv6.cpp
+++ b/Engine/Source/Runtime/Sockets/Private/BSDIPv6Sockets/SocketSubsystemBSDIPv6.cpp
@@ -8,10 +8,10 @@
 #include "BSDIPv6Sockets/IPAddressBSDIPv6.h"
 
 PRAGMA_DISABLE_DEPRECATION_WARNINGS
-class FSocketBSDIPv6* FSocketSubsystemBSDIPv6::InternalBSDSocketFactory(SOCKET Socket, ESocketType SocketType, const FString& SocketDescription)
+class FSocketBSDIPv6* FSocketSubsystemBSDIPv6::InternalBSDSocketFactory(SOCKET Socket, ESocketType SocketType, const FString& SocketDescription, ESocketProtocolFamily SocketProtocol)
 {
 	// return a new socket object
-	return new FSocketBSDIPv6(Socket, SocketType, SocketDescription, this);
+	return new FSocketBSDIPv6(Socket, SocketType, SocketDescription, SocketProtocol, this);
 }
 
 ESocketErrors FSocketSubsystemBSDIPv6::TranslateGAIErrorCode(int32 Code) const
@@ -50,17 +50,24 @@ FSocket* FSocketSubsystemBSDIPv6::CreateSocket(const FName& SocketType, const FS
 {
 	SOCKET Socket = INVALID_SOCKET;
 	FSocket* NewSocket = NULL;
+
+	// If we're passed None, initialize with IPv6 to maintain valid values.
+	if (ProtocolType == ESocketProtocolFamily::None)
+	{
+		ProtocolType = ESocketProtocolFamily::IPv6;
+	}
+
 	switch (SocketType.GetComparisonIndex())
 	{
 	case NAME_DGram:
 		// Creates a data gram (UDP) socket
 		Socket = socket(AF_INET6, SOCK_DGRAM, IPPROTO_UDP);
-		NewSocket = (Socket != INVALID_SOCKET) ? InternalBSDSocketFactory(Socket, SOCKTYPE_Datagram, SocketDescription) : NULL;
+		NewSocket = (Socket != INVALID_SOCKET) ? InternalBSDSocketFactory(Socket, SOCKTYPE_Datagram, SocketDescription, ProtocolType) : NULL;
 		break;
 	case NAME_Stream:
 		// Creates a stream (TCP) socket
 		Socket = socket(AF_INET6, SOCK_STREAM, IPPROTO_TCP);
-		NewSocket = (Socket != INVALID_SOCKET) ? InternalBSDSocketFactory(Socket, SOCKTYPE_Streaming, SocketDescription) : NULL;
+		NewSocket = (Socket != INVALID_SOCKET) ? InternalBSDSocketFactory(Socket, SOCKTYPE_Streaming, SocketDescription, ProtocolType) : NULL;
 		break;
 	}
 
diff --git a/Engine/Source/Runtime/Sockets/Private/BSDIPv6Sockets/SocketSubsystemBSDIPv6.h b/Engine/Source/Runtime/Sockets/Private/BSDIPv6Sockets/SocketSubsystemBSDIPv6.h
--- a/Engine/Source/Runtime/Sockets/Private/BSDIPv6Sockets/SocketSubsystemBSDIPv6.h
+++ b/Engine/Source/Runtime/Sockets/Private/BSDIPv6Sockets/SocketSubsystemBSDIPv6.h
@@ -71,7 +71,17 @@ protected:
 	/**
 	 * Allows a subsystem subclass to create a FSocketBSD sub class.
 	 */
-	virtual class FSocketBSDIPv6* InternalBSDSocketFactory( SOCKET Socket, ESocketType SocketType, const FString& SocketDescription );
+	virtual class FSocketBSDIPv6* InternalBSDSocketFactory( SOCKET Socket, ESocketType SocketType, const FString& SocketDescription, ESocketProtocolFamily SocketProtocol );
+
+
+	/**
+	 * Allows a subsystem subclass to create a FSocketBSD sub class.
+	 */
+	DEPRECATED(4.22, "To support multiple stack types, move to the constructor that allows for specifying the protocol stack to initialize the socket on.")
+	virtual class FSocketBSD* InternalBSDSocketFactory(SOCKET Socket, ESocketType SocketType, const FString& SocketDescription)
+	{
+		return InternalBSDSocketFactory(Socket, SocketType, SocketDescription, ESocketProtocolFamily::IPv6);
+	}
 
 	/**
 	 * Translates an ESocketAddressInfoFlags into a value usable by getaddrinfo
diff --git a/Engine/Source/Runtime/Sockets/Private/BSDIPv6Sockets/SocketsBSDIPv6.cpp b/Engine/Source/Runtime/Sockets/Private/BSDIPv6Sockets/SocketsBSDIPv6.cpp
--- a/Engine/Source/Runtime/Sockets/Private/BSDIPv6Sockets/SocketsBSDIPv6.cpp
+++ b/Engine/Source/Runtime/Sockets/Private/BSDIPv6Sockets/SocketsBSDIPv6.cpp
@@ -142,7 +142,7 @@ FSocket* FSocketBSDIPv6::Accept(const FString& InSocketDescription)
 		// we need the subclass to create the actual FSocket object
 		check(SocketSubsystem);
 		FSocketSubsystemBSDIPv6* BSDSystem = static_cast<FSocketSubsystemBSDIPv6*>(SocketSubsystem);
-		return BSDSystem->InternalBSDSocketFactory(NewSocket, SocketType, InSocketDescription);
+		return BSDSystem->InternalBSDSocketFactory(NewSocket, SocketType, InSocketDescription, SocketProtocol);
 	}
 
 	return NULL;
@@ -159,7 +159,7 @@ FSocket* FSocketBSDIPv6::Accept(FInternetAddr& OutAddr, const FString& InSocketD
 		// we need the subclass to create the actual FSocket object
 		check(SocketSubsystem);
 		FSocketSubsystemBSDIPv6* BSDSystem = static_cast<FSocketSubsystemBSDIPv6*>(SocketSubsystem);
-		return BSDSystem->InternalBSDSocketFactory(NewSocket, SocketType, InSocketDescription);
+		return BSDSystem->InternalBSDSocketFactory(NewSocket, SocketType, InSocketDescription, SocketProtocol);
 	}
 
 	return NULL;
diff --git a/Engine/Source/Runtime/Sockets/Private/BSDIPv6Sockets/SocketsBSDIPv6.h b/Engine/Source/Runtime/Sockets/Private/BSDIPv6Sockets/SocketsBSDIPv6.h
--- a/Engine/Source/Runtime/Sockets/Private/BSDIPv6Sockets/SocketsBSDIPv6.h
+++ b/Engine/Source/Runtime/Sockets/Private/BSDIPv6Sockets/SocketsBSDIPv6.h
@@ -48,8 +48,24 @@ public:
 	 * @param InSocketDescription the debug description of the socket
 	 * @param InSubsystem the subsystem that created this socket
 	 */
-	FSocketBSDIPv6(SOCKET InSocket, ESocketType InSocketType, const FString& InSocketDescription, ISocketSubsystem* InSubsystem) 
-		: FSocket(InSocketType, InSocketDescription)
+	DEPRECATED(4.22, "Use the socket constructor with protocol specification.")
+	FSocketBSDIPv6(SOCKET InSocket, ESocketType InSocketType, const FString& InSocketDescription, ISocketSubsystem* InSubsystem)
+		: FSocket(InSocketType, InSocketDescription, ESocketProtocolFamily::IPv6)
+		, Socket(InSocket)
+		, SocketSubsystem(InSubsystem)
+	{ }
+
+	/**
+	 * Assigns a BSD socket to this object
+	 *
+	 * @param InSocket the socket to assign to this object
+	 * @param InSocketType the type of socket that was created
+	 * @param InSocketDescription the debug description of the socket
+	 * @param InSocketProtocol the protocol used to create this socket.
+	 * @param InSubsystem the subsystem that created this socket
+	 */
+	FSocketBSDIPv6(SOCKET InSocket, ESocketType InSocketType, const FString& InSocketDescription, ESocketProtocolFamily InSocketProtocol, ISocketSubsystem* InSubsystem) 
+		: FSocket(InSocketType, InSocketDescription, InSocketProtocol)
 		, Socket(InSocket)
 		, SocketSubsystem(InSubsystem)
 	{ }
diff --git a/Engine/Source/Runtime/Sockets/Private/BSDSockets/SocketSubsystemBSD.cpp b/Engine/Source/Runtime/Sockets/Private/BSDSockets/SocketSubsystemBSD.cpp
--- a/Engine/Source/Runtime/Sockets/Private/BSDSockets/SocketSubsystemBSD.cpp
+++ b/Engine/Source/Runtime/Sockets/Private/BSDSockets/SocketSubsystemBSD.cpp
@@ -9,10 +9,10 @@
 #include "BSDSockets/SocketsBSD.h"
 #include <errno.h>
 
-FSocketBSD* FSocketSubsystemBSD::InternalBSDSocketFactory(SOCKET Socket, ESocketType SocketType, const FString& SocketDescription)
+FSocketBSD* FSocketSubsystemBSD::InternalBSDSocketFactory(SOCKET Socket, ESocketType SocketType, const FString& SocketDescription, ESocketProtocolFamily SocketProtocol)
 {
 	// return a new socket object
-	return new FSocketBSD(Socket, SocketType, SocketDescription, this);
+	return new FSocketBSD(Socket, SocketType, SocketDescription, SocketProtocol, this);
 }
 
 FSocket* FSocketSubsystemBSD::CreateSocket(const FName& SocketType, const FString& SocketDescription, bool bForceUDP)
@@ -48,7 +48,7 @@ FSocket* FSocketSubsystemBSD::CreateSocket(const FName& SocketType, const FStrin
 	int32 SocketTypeFlag = (bIsUDP) ? SOCK_DGRAM : SOCK_STREAM;
 
 	Socket = socket(GetProtocolFamilyValue(ProtocolType), SocketTypeFlag | PlatformSpecificTypeFlags, ((bIsUDP) ? IPPROTO_UDP : IPPROTO_TCP));
-	NewSocket = (Socket != INVALID_SOCKET) ? InternalBSDSocketFactory(Socket, ((bIsUDP) ? SOCKTYPE_Datagram : SOCKTYPE_Streaming), SocketDescription) : nullptr;
+	NewSocket = (Socket != INVALID_SOCKET) ? InternalBSDSocketFactory(Socket, ((bIsUDP) ? SOCKTYPE_Datagram : SOCKTYPE_Streaming), SocketDescription, ProtocolType) : nullptr;
 
 	if (!NewSocket)
 	{
diff --git a/Engine/Source/Runtime/Sockets/Private/BSDSockets/SocketSubsystemBSD.h b/Engine/Source/Runtime/Sockets/Private/BSDSockets/SocketSubsystemBSD.h
--- a/Engine/Source/Runtime/Sockets/Private/BSDSockets/SocketSubsystemBSD.h
+++ b/Engine/Source/Runtime/Sockets/Private/BSDSockets/SocketSubsystemBSD.h
@@ -105,7 +105,16 @@ protected:
 	/**
 	 * Allows a subsystem subclass to create a FSocketBSD sub class.
 	 */
-	virtual class FSocketBSD* InternalBSDSocketFactory( SOCKET Socket, ESocketType SocketType, const FString& SocketDescription );
+	virtual class FSocketBSD* InternalBSDSocketFactory( SOCKET Socket, ESocketType SocketType, const FString& SocketDescription, ESocketProtocolFamily SocketProtocol);
+
+	/**
+	 * Allows a subsystem subclass to create a FSocketBSD sub class.
+	 */
+	DEPRECATED(4.22, "To support multiple stack types, move to the constructor that allows for specifying the protocol stack to initialize the socket on.")
+	virtual class FSocketBSD* InternalBSDSocketFactory(SOCKET Socket, ESocketType SocketType, const FString& SocketDescription)
+	{
+		return InternalBSDSocketFactory(Socket, SocketType, SocketDescription, GetDefaultSocketProtocolFamily());
+	}
 
 	// allow BSD sockets to use this when creating new sockets from accept() etc
 	friend FSocketBSD;
diff --git a/Engine/Source/Runtime/Sockets/Private/BSDSockets/SocketsBSD.cpp b/Engine/Source/Runtime/Sockets/Private/BSDSockets/SocketsBSD.cpp
--- a/Engine/Source/Runtime/Sockets/Private/BSDSockets/SocketsBSD.cpp
+++ b/Engine/Source/Runtime/Sockets/Private/BSDSockets/SocketsBSD.cpp
@@ -103,7 +103,7 @@ FSocket* FSocketBSD::Accept(const FString& InSocketDescription)
 		// we need the subclass to create the actual FSocket object
 		check(SocketSubsystem);
 		FSocketSubsystemBSD* BSDSystem = static_cast<FSocketSubsystemBSD*>(SocketSubsystem);
-		return BSDSystem->InternalBSDSocketFactory(NewSocket, SocketType, InSocketDescription);
+		return BSDSystem->InternalBSDSocketFactory(NewSocket, SocketType, InSocketDescription, SocketProtocol);
 	}
 
 	return NULL;
@@ -121,7 +121,7 @@ FSocket* FSocketBSD::Accept(FInternetAddr& OutAddr, const FString& InSocketDescr
 		// we need the subclass to create the actual FSocket object
 		check(SocketSubsystem);
 		FSocketSubsystemBSD* BSDSystem = static_cast<FSocketSubsystemBSD*>(SocketSubsystem);
-		return BSDSystem->InternalBSDSocketFactory(NewSocket, SocketType, InSocketDescription);
+		return BSDSystem->InternalBSDSocketFactory(NewSocket, SocketType, InSocketDescription, BSDAddr.GetProtocolFamily());
 	}
 
 	return NULL;
diff --git a/Engine/Source/Runtime/Sockets/Private/BSDSockets/SocketsBSD.h b/Engine/Source/Runtime/Sockets/Private/BSDSockets/SocketsBSD.h
--- a/Engine/Source/Runtime/Sockets/Private/BSDSockets/SocketsBSD.h
+++ b/Engine/Source/Runtime/Sockets/Private/BSDSockets/SocketsBSD.h
@@ -49,8 +49,24 @@ public:
 	 * @param InSocketType the type of socket that was created.
 	 * @param InSocketDescription the debug description of the socket.
 	 */
-	FSocketBSD(SOCKET InSocket, ESocketType InSocketType, const FString& InSocketDescription, ISocketSubsystem * InSubsystem) 
-		: FSocket(InSocketType, InSocketDescription)
+	DEPRECATED(4.22, "Use the socket constructor that specifies protocol stack for better compatibility and debugging")
+	FSocketBSD(SOCKET InSocket, ESocketType InSocketType, const FString& InSocketDescription, ISocketSubsystem * InSubsystem)
+		: FSocket(InSocketType, InSocketDescription, ESocketProtocolFamily::None)
+		, Socket(InSocket)
+		, LastActivityTime(0)
+		, SocketSubsystem(InSubsystem)
+	{ }
+
+	/**
+	 * Assigns a BSD socket to this object.
+	 *
+	 * @param InSocket the socket to assign to this object.
+	 * @param InSocketType the type of socket that was created.
+	 * @param InSocketDescription the debug description of the socket.
+	 * @param InSocketProtocol the protocol this socket is initialized with
+	 */
+	FSocketBSD(SOCKET InSocket, ESocketType InSocketType, const FString& InSocketDescription, ESocketProtocolFamily InSocketProtocol, ISocketSubsystem * InSubsystem)
+		: FSocket(InSocketType, InSocketDescription, InSocketProtocol)
 		, Socket(InSocket)
 		, LastActivityTime(0)
 		, SocketSubsystem(InSubsystem)
diff --git a/Engine/Source/Runtime/Sockets/Private/IOS/SocketSubsystemIOS.cpp b/Engine/Source/Runtime/Sockets/Private/IOS/SocketSubsystemIOS.cpp
--- a/Engine/Source/Runtime/Sockets/Private/IOS/SocketSubsystemIOS.cpp
+++ b/Engine/Source/Runtime/Sockets/Private/IOS/SocketSubsystemIOS.cpp
@@ -12,10 +12,10 @@
 
 FSocketSubsystemIOS* FSocketSubsystemIOS::SocketSingleton = NULL;
 
-class FSocketBSD* FSocketSubsystemIOS::InternalBSDSocketFactory(SOCKET Socket, ESocketType SocketType, const FString& SocketDescription)
+class FSocketBSD* FSocketSubsystemIOS::InternalBSDSocketFactory(SOCKET Socket, ESocketType SocketType, const FString& SocketDescription, ESocketProtocolFamily SocketProtocol)
 {
 	UE_LOG(LogIOS, Log, TEXT(" FSocketSubsystemIOS::InternalBSDSocketFactory"));
-	return new FSocketBSDIOS(Socket, SocketType, SocketDescription, this);
+	return new FSocketBSDIOS(Socket, SocketType, SocketDescription, SocketProtocol, this);
 }
 
 FName CreateSocketSubsystem( FSocketSubsystemModule& SocketSubsystemModule )
diff --git a/Engine/Source/Runtime/Sockets/Private/IOS/SocketSubsystemIOS.h b/Engine/Source/Runtime/Sockets/Private/IOS/SocketSubsystemIOS.h
--- a/Engine/Source/Runtime/Sockets/Private/IOS/SocketSubsystemIOS.h
+++ b/Engine/Source/Runtime/Sockets/Private/IOS/SocketSubsystemIOS.h
@@ -60,5 +60,5 @@ public:
 	virtual FSocket* CreateSocket(const FName& SocketType, const FString& SocketDescription, ESocketProtocolFamily ProtocolType) override;
 
 	virtual TSharedRef<FInternetAddr> GetLocalHostAddr(FOutputDevice& Out, bool& bCanBindAll) override;
-	virtual class FSocketBSD* InternalBSDSocketFactory(SOCKET Socket, ESocketType SocketType, const FString& SocketDescription) override;
+	virtual class FSocketBSD* InternalBSDSocketFactory(SOCKET Socket, ESocketType SocketType, const FString& SocketDescription, ESocketProtocolFamily SocketProtocol) override;
 };
\ No newline at end of file
diff --git a/Engine/Source/Runtime/Sockets/Private/IOS/SocketsBSDIOS.h b/Engine/Source/Runtime/Sockets/Private/IOS/SocketsBSDIOS.h
--- a/Engine/Source/Runtime/Sockets/Private/IOS/SocketsBSDIOS.h
+++ b/Engine/Source/Runtime/Sockets/Private/IOS/SocketsBSDIOS.h
@@ -20,8 +20,8 @@ class FSocketBSDIOS
 {
 public:
 
-	FSocketBSDIOS(SOCKET InSocket, ESocketType InSocketType, const FString& InSocketDescription, ISocketSubsystem* InSubsystem)
-		:FSocketBSD(InSocket, InSocketType, InSocketDescription, InSubsystem)
+	FSocketBSDIOS(SOCKET InSocket, ESocketType InSocketType, const FString& InSocketDescription, ESocketProtocolFamily InSocketProtocol, ISocketSubsystem* InSubsystem)
+		:FSocketBSD(InSocket, InSocketType, InSocketDescription, InSocketProtocol, InSubsystem)
 	{
 	}
 
diff --git a/Engine/Source/Runtime/Sockets/Private/Mac/SocketSubsystemMac.cpp b/Engine/Source/Runtime/Sockets/Private/Mac/SocketSubsystemMac.cpp
--- a/Engine/Source/Runtime/Sockets/Private/Mac/SocketSubsystemMac.cpp
+++ b/Engine/Source/Runtime/Sockets/Private/Mac/SocketSubsystemMac.cpp
@@ -66,10 +66,10 @@ bool FSocketSubsystemMac::HasNetworkDevice()
 	return true;
 }
 
-class FSocketBSD* FSocketSubsystemMac::InternalBSDSocketFactory(SOCKET Socket, ESocketType SocketType, const FString& SocketDescription)
+class FSocketBSD* FSocketSubsystemMac::InternalBSDSocketFactory(SOCKET Socket, ESocketType SocketType, const FString& SocketDescription, ESocketProtocolFamily SocketProtocol)
 {
 	// return a new socket object
-	FSocketMac* MacSock = new FSocketMac(Socket, SocketType, SocketDescription, this);
+	FSocketMac* MacSock = new FSocketMac(Socket, SocketType, SocketDescription, SocketProtocol, this);
 	if (MacSock)
 	{
 		// disable the SIGPIPE exception
diff --git a/Engine/Source/Runtime/Sockets/Private/Mac/SocketSubsystemMac.h b/Engine/Source/Runtime/Sockets/Private/Mac/SocketSubsystemMac.h
--- a/Engine/Source/Runtime/Sockets/Private/Mac/SocketSubsystemMac.h
+++ b/Engine/Source/Runtime/Sockets/Private/Mac/SocketSubsystemMac.h
@@ -36,7 +36,7 @@ PACKAGE_SCOPE:
 	/**
 	 * Allows a subsystem subclass to create a FSocketBSD sub class 
 	 */
-	virtual class FSocketBSD* InternalBSDSocketFactory(SOCKET Socket, ESocketType SocketType, const FString& SocketDescription);
+	virtual class FSocketBSD* InternalBSDSocketFactory(SOCKET Socket, ESocketType SocketType, const FString& SocketDescription, ESocketProtocolFamily SocketProtocol);
 
 public:
 
@@ -83,8 +83,8 @@ public:
 	 * @param InSocketType the type of socket that was created
 	 * @param InSocketDescription the debug description of the socket
 	 */
-	FSocketMac(SOCKET InSocket, ESocketType InSocketType, const FString& InSocketDescription, ISocketSubsystem * InSubsystem) 
-	: FSocketBSD(InSocket, InSocketType, InSocketDescription, InSubsystem)
+	FSocketMac(SOCKET InSocket, ESocketType InSocketType, const FString& InSocketDescription, ESocketProtocolFamily InSocketProtocol, ISocketSubsystem * InSubsystem) 
+	: FSocketBSD(InSocket, InSocketType, InSocketDescription, InSocketProtocol, InSubsystem)
 	{
 	}
 
diff --git a/Engine/Source/Runtime/Sockets/Public/Sockets.h b/Engine/Source/Runtime/Sockets/Public/Sockets.h
--- a/Engine/Source/Runtime/Sockets/Public/Sockets.h
+++ b/Engine/Source/Runtime/Sockets/Public/Sockets.h
@@ -20,12 +20,16 @@ protected:
 	/** Debug description of socket usage. */
 	FString SocketDescription;
 
+	/** Protocol used in creation of a socket */
+	ESocketProtocolFamily SocketProtocol;
+
 public:
 
 	/** Default constructor. */
 	inline FSocket() :
 		SocketType(SOCKTYPE_Unknown),
-		SocketDescription(TEXT(""))
+		SocketDescription(TEXT("")),
+		SocketProtocol(ESocketProtocolFamily::None)
 	{ }
 
 	/**
@@ -34,9 +38,24 @@ public:
 	 * @param InSocketType The type of socket being created
 	 * @param InSocketDescription The debug description of the socket
 	 */
+	DEPRECATED(4.22, "Please migrate to the constructor that specifies protocol stack")
 	inline FSocket(ESocketType InSocketType, const FString& InSocketDescription) :
 		SocketType(InSocketType),
-		SocketDescription(InSocketDescription)
+		SocketDescription(InSocketDescription),
+		SocketProtocol(ESocketProtocolFamily::None)
+	{ }
+
+	/**
+	 * Specifies the type of socket being created
+	 *
+	 * @param InSocketType The type of socket being created
+	 * @param InSocketDescription The debug description of the socket
+	 * @param InSocketProtocol the protocol stack this socket should be created on.
+	 */
+	inline FSocket(ESocketType InSocketType, const FString& InSocketDescription, ESocketProtocolFamily InSocketProtocol) :
+		SocketType(InSocketType),
+		SocketDescription(InSocketDescription),
+		SocketProtocol(InSocketProtocol)
 	{ }
 
 	/** Virtual destructor. */
@@ -335,4 +354,15 @@ public:
 	{
 		return SocketDescription;
 	}
+
+	/**
+	 * Get the type of protocol the socket is bound to
+	 *
+	 * @return Socket type.
+	 * @see GetDescription, GetPortNo
+	 */
+	FORCEINLINE ESocketProtocolFamily GetProtocol() const
+	{
+		return SocketProtocol;
+	}
 };
