Fix iOS compile issues

#rnx

diff --git a/Engine/Source/Runtime/Sockets/Private/BSDIPv6Sockets/SocketSubsystemBSDIPv6.h b/Engine/Source/Runtime/Sockets/Private/BSDIPv6Sockets/SocketSubsystemBSDIPv6.h
--- a/Engine/Source/Runtime/Sockets/Private/BSDIPv6Sockets/SocketSubsystemBSDIPv6.h
+++ b/Engine/Source/Runtime/Sockets/Private/BSDIPv6Sockets/SocketSubsystemBSDIPv6.h
@@ -78,7 +78,7 @@ protected:
 	 * Allows a subsystem subclass to create a FSocketBSD sub class.
 	 */
 	DEPRECATED(4.22, "To support multiple stack types, move to the constructor that allows for specifying the protocol stack to initialize the socket on.")
-	virtual class FSocketBSD* InternalBSDSocketFactory(SOCKET Socket, ESocketType SocketType, const FString& SocketDescription)
+	virtual class FSocketBSDIPv6* InternalBSDSocketFactory(SOCKET Socket, ESocketType SocketType, const FString& SocketDescription)
 	{
 		return InternalBSDSocketFactory(Socket, SocketType, SocketDescription, ESocketProtocolFamily::IPv6);
 	}
