Fix incorrect endian parsing in SetIP(uint) due to bad format.

#rnx

diff --git a/Engine/Source/Runtime/Sockets/Private/BSDSockets/IPAddressBSD.cpp b/Engine/Source/Runtime/Sockets/Private/BSDSockets/IPAddressBSD.cpp
--- a/Engine/Source/Runtime/Sockets/Private/BSDSockets/IPAddressBSD.cpp
+++ b/Engine/Source/Runtime/Sockets/Private/BSDSockets/IPAddressBSD.cpp
@@ -196,7 +196,7 @@ void FInternetAddrBSD::SetIp(uint32 InAddr)
 		else
 		{
 			in6_addr ConvertedAddrData;
-			MapIPv4ToIPv6(InAddr, ConvertedAddrData);
+			MapIPv4ToIPv6(htonl(InAddr), ConvertedAddrData);
 			SetIp(ConvertedAddrData);
 		}
 		return;
