Fix memberships on iOS.

#rnx

diff --git a/Engine/Source/Runtime/Sockets/Private/BSDSockets/SocketsBSD.cpp b/Engine/Source/Runtime/Sockets/Private/BSDSockets/SocketsBSD.cpp
--- a/Engine/Source/Runtime/Sockets/Private/BSDSockets/SocketsBSD.cpp
+++ b/Engine/Source/Runtime/Sockets/Private/BSDSockets/SocketsBSD.cpp
@@ -386,9 +386,9 @@ bool FSocketBSD::JoinMulticastGroup(const FInternetAddr& GroupAddress)
 	if (BSDAddr.GetProtocolFamily() == ESocketProtocolFamily::IPv6)
 	{
 		ipv6_mreq imr;
-		imr.ipv6mr_interface = 0;
+		imr.ipv6mr_interface = htonl(BSDAddr.GetScopeId());
 		imr.ipv6mr_multiaddr = ((sockaddr_in6*)&(BSDAddr.Addr))->sin6_addr;
-		return (setsockopt(Socket, IPPROTO_IPV6, IP_ADD_MEMBERSHIP, (char*)&imr, sizeof(imr)) == 0);
+		return (setsockopt(Socket, IPPROTO_IPV6, IPV6_JOIN_GROUP, (char*)&imr, sizeof(imr)) == 0);
 	}
 #endif
 
@@ -407,9 +407,9 @@ bool FSocketBSD::LeaveMulticastGroup(const FInternetAddr& GroupAddress)
 	if (BSDAddr.GetProtocolFamily() == ESocketProtocolFamily::IPv6)
 	{
 		ipv6_mreq imr;
-		imr.ipv6mr_interface = 0;
+		imr.ipv6mr_interface = htonl(BSDAddr.GetScopeId());
 		imr.ipv6mr_multiaddr = ((sockaddr_in6*)&(BSDAddr.Addr))->sin6_addr;
-		return (setsockopt(Socket, IPPROTO_IPV6, IP_DROP_MEMBERSHIP, (char*)&imr, sizeof(imr)) == 0);
+		return (setsockopt(Socket, IPPROTO_IPV6, IPV6_LEAVE_GROUP, (char*)&imr, sizeof(imr)) == 0);
 	}
 #endif
 
@@ -425,7 +425,8 @@ bool FSocketBSD::SetMulticastLoopback(bool bLoopback)
 #if PLATFORM_HAS_BSD_IPV6_SOCKETS
 	if (SocketProtocol == ESocketProtocolFamily::IPv6)
 	{
-		return (setsockopt(Socket, IPPROTO_IPV6, IP_MULTICAST_LOOP, (char*)&bLoopback, sizeof(bLoopback)) == 0);
+		uint32 ShouldLoopback = bLoopback ? 1 : 0;
+		return (setsockopt(Socket, IPPROTO_IPV6, IP_MULTICAST_LOOP, (char*)&ShouldLoopback, sizeof(ShouldLoopback)) == 0);
 	}
 #endif
 
@@ -438,7 +439,8 @@ bool FSocketBSD::SetMulticastTtl(uint8 TimeToLive)
 #if PLATFORM_HAS_BSD_IPV6_SOCKETS
 	if (SocketProtocol == ESocketProtocolFamily::IPv6)
 	{
-		return (setsockopt(Socket, IPPROTO_IPV6, IP_MULTICAST_TTL, (char*)&TimeToLive, sizeof(TimeToLive)) == 0);
+		uint32 RealTimeToLive = TimeToLive;
+		return (setsockopt(Socket, IPPROTO_IPV6, IP_MULTICAST_TTL, (char*)&RealTimeToLive, sizeof(RealTimeToLive)) == 0);
 	}
 #endif
 
@@ -533,7 +535,7 @@ bool FSocketBSD::SetIPv6Only(bool bIPv6Only)
 {
 #if PLATFORM_HAS_BSD_IPV6_SOCKETS
 	int v6only = bIPv6Only ? 1 : 0;
-	bool bOk = setsockopt(Socket, IPPROTO_IPV6, IPV6_V6ONLY, (char*)&v6only, sizeof(v6only)) == 0;
+	bool bOk = (setsockopt(Socket, IPPROTO_IPV6, IPV6_V6ONLY, (char*)&v6only, sizeof(v6only)) == 0);
 
 	if (bOk == false)
 	{
diff --git a/Engine/Source/Runtime/Sockets/Private/IOS/IPAddressBSDIOS.h b/Engine/Source/Runtime/Sockets/Private/IOS/IPAddressBSDIOS.h
--- a/Engine/Source/Runtime/Sockets/Private/IOS/IPAddressBSDIOS.h
+++ b/Engine/Source/Runtime/Sockets/Private/IOS/IPAddressBSDIOS.h
@@ -26,6 +26,13 @@ public:
 			{
 				UE_LOG(LogSockets, Warning, TEXT("Could not resolve the broadcast address for iOS, this address will just be blank"));
 			}
+			else
+			{
+				bool bUnusedBool;
+				TSharedRef<FInternetAddrBSDIOS> ScopeAddr = 
+					StaticCastSharedRef<FInternetAddrBSDIOS>(SocketSubsystem->GetLocalHostAddr(*GLog, bUnusedBool));
+				SetScopeId(ScopeAddr->GetScopeId());
+			}
 		}
 		else
 		{
