Fix deprecated DEPRECATED macros.

diff --git a/Engine/Source/Runtime/Sockets/Private/BSDIPv6Sockets/SocketSubsystemBSDIPv6.h b/Engine/Source/Runtime/Sockets/Private/BSDIPv6Sockets/SocketSubsystemBSDIPv6.h
--- a/Engine/Source/Runtime/Sockets/Private/BSDIPv6Sockets/SocketSubsystemBSDIPv6.h
+++ b/Engine/Source/Runtime/Sockets/Private/BSDIPv6Sockets/SocketSubsystemBSDIPv6.h
@@ -79,7 +79,7 @@ protected:
 	/**
 	 * Allows a subsystem subclass to create a FSocketBSD sub class.
 	 */
-	DEPRECATED(4.22, "To support multiple stack types, move to the constructor that allows for specifying the protocol stack to initialize the socket on.")
+	UE_DEPRECATED(4.22, "To support multiple stack types, move to the constructor that allows for specifying the protocol stack to initialize the socket on.")
 	virtual class FSocketBSDIPv6* InternalBSDSocketFactory(SOCKET Socket, ESocketType SocketType, const FString& SocketDescription)
 	{
 		return InternalBSDSocketFactory(Socket, SocketType, SocketDescription, ESocketProtocolFamily::IPv6);
diff --git a/Engine/Source/Runtime/Sockets/Private/BSDIPv6Sockets/SocketsBSDIPv6.h b/Engine/Source/Runtime/Sockets/Private/BSDIPv6Sockets/SocketsBSDIPv6.h
--- a/Engine/Source/Runtime/Sockets/Private/BSDIPv6Sockets/SocketsBSDIPv6.h
+++ b/Engine/Source/Runtime/Sockets/Private/BSDIPv6Sockets/SocketsBSDIPv6.h
@@ -48,7 +48,7 @@ public:
 	 * @param InSocketDescription the debug description of the socket
 	 * @param InSubsystem the subsystem that created this socket
 	 */
-	DEPRECATED(4.22, "Use the socket constructor with protocol specification.")
+	UE_DEPRECATED(4.22, "Use the socket constructor with protocol specification.")
 	FSocketBSDIPv6(SOCKET InSocket, ESocketType InSocketType, const FString& InSocketDescription, ISocketSubsystem* InSubsystem)
 		: FSocket(InSocketType, InSocketDescription, ESocketProtocolFamily::IPv6)
 		, Socket(InSocket)
diff --git a/Engine/Source/Runtime/Sockets/Private/BSDSockets/SocketSubsystemBSD.h b/Engine/Source/Runtime/Sockets/Private/BSDSockets/SocketSubsystemBSD.h
--- a/Engine/Source/Runtime/Sockets/Private/BSDSockets/SocketSubsystemBSD.h
+++ b/Engine/Source/Runtime/Sockets/Private/BSDSockets/SocketSubsystemBSD.h
@@ -112,7 +112,7 @@ protected:
 	/**
 	 * Allows a subsystem subclass to create a FSocketBSD sub class.
 	 */
-	DEPRECATED(4.22, "To support multiple stack types, move to the constructor that allows for specifying the protocol stack to initialize the socket on.")
+	UE_DEPRECATED(4.22, "To support multiple stack types, move to the constructor that allows for specifying the protocol stack to initialize the socket on.")
 	virtual class FSocketBSD* InternalBSDSocketFactory(SOCKET Socket, ESocketType SocketType, const FString& SocketDescription)
 	{
 		return InternalBSDSocketFactory(Socket, SocketType, SocketDescription, GetDefaultSocketProtocolFamily());
diff --git a/Engine/Source/Runtime/Sockets/Private/BSDSockets/SocketsBSD.h b/Engine/Source/Runtime/Sockets/Private/BSDSockets/SocketsBSD.h
--- a/Engine/Source/Runtime/Sockets/Private/BSDSockets/SocketsBSD.h
+++ b/Engine/Source/Runtime/Sockets/Private/BSDSockets/SocketsBSD.h
@@ -50,7 +50,7 @@ public:
 	 * @param InSocketType the type of socket that was created.
 	 * @param InSocketDescription the debug description of the socket.
 	 */
-	DEPRECATED(4.22, "Use the socket constructor that specifies protocol stack for better compatibility and debugging")
+	UE_DEPRECATED(4.22, "Use the socket constructor that specifies protocol stack for better compatibility and debugging")
 	FSocketBSD(SOCKET InSocket, ESocketType InSocketType, const FString& InSocketDescription, ISocketSubsystem * InSubsystem)
 		: FSocket(InSocketType, InSocketDescription, ESocketProtocolFamily::None)
 		, Socket(InSocket)
diff --git a/Engine/Source/Runtime/Sockets/Public/Sockets.h b/Engine/Source/Runtime/Sockets/Public/Sockets.h
--- a/Engine/Source/Runtime/Sockets/Public/Sockets.h
+++ b/Engine/Source/Runtime/Sockets/Public/Sockets.h
@@ -38,7 +38,7 @@ public:
 	 * @param InSocketType The type of socket being created
 	 * @param InSocketDescription The debug description of the socket
 	 */
-	DEPRECATED(4.22, "Please migrate to the constructor that specifies protocol stack")
+	UE_DEPRECATED(4.22, "Please migrate to the constructor that specifies protocol stack")
 	inline FSocket(ESocketType InSocketType, const FString& InSocketDescription) :
 		SocketType(InSocketType),
 		SocketDescription(InSocketDescription),
