Fix Android GetLocalHostAddr from stomping on address information when we're determining best address usage. Also prevents the invalid ALL|V4MAPPED flag from being attached to GAI queries.

diff --git a/Engine/Source/Runtime/Sockets/Private/Android/SocketSubsystemAndroid.cpp b/Engine/Source/Runtime/Sockets/Private/Android/SocketSubsystemAndroid.cpp
--- a/Engine/Source/Runtime/Sockets/Private/Android/SocketSubsystemAndroid.cpp
+++ b/Engine/Source/Runtime/Sockets/Private/Android/SocketSubsystemAndroid.cpp
@@ -170,6 +170,9 @@ TSharedRef<FInternetAddr> FSocketSubsystemAndroid::GetLocalHostAddr(FOutputDevic
 		int Result = ioctl(TempSocket, SIOCGIFCONF, &IfConfig);
 		if (Result == 0)
 		{
+			// Temporary cache of the address we get per interface
+			sockaddr_storage TemporaryAddress;
+			// Interface specific address containers
 			sockaddr_storage WifiAddress;
 			sockaddr_storage CellularAddress;
 			sockaddr_storage OtherAddress;
@@ -181,27 +184,30 @@ TSharedRef<FInternetAddr> FSocketSubsystemAndroid::GetLocalHostAddr(FOutputDevic
 
 			for (int32 IdxReq = 0; IdxReq < ARRAY_COUNT(IfReqs); ++IdxReq)
 			{
+				// Cache the address information, as the following flag lookup will 
+				// write into the ifr_addr field.
+				FMemory::Memcpy((void*)&TemporaryAddress, (void*)(&IfReqs[IdxReq].ifr_addr), sizeof(sockaddr_in));
+
 				// Examine interfaces that are up and not loop back
-				int ResultFlags = ioctl(TempSocket, SIOCGIFFLAGS, &IfReqs[IdxReq]);
-				if (ResultFlags == 0 &&
+				if (ioctl(TempSocket, SIOCGIFFLAGS, &IfReqs[IdxReq]) == 0 &&
 					(IfReqs[IdxReq].ifr_flags & IFF_UP) &&
 					(IfReqs[IdxReq].ifr_flags & IFF_LOOPBACK) == 0)
 				{
 					if (strcmp(IfReqs[IdxReq].ifr_name, "wlan0") == 0)
 					{
 						// 'Usually' wifi, Prefer wifi
-						FMemory::Memcpy((void*)&WifiAddress, (void*)(&IfReqs[IdxReq].ifr_addr), sizeof(sockaddr_in));
+						FMemory::Memcpy((void*)&WifiAddress, (void*)(&TemporaryAddress), sizeof(sockaddr_in));
 						break;
 					}
 					else if (strcmp(IfReqs[IdxReq].ifr_name, "rmnet0") == 0)
 					{
 						// 'Usually' cellular
-						FMemory::Memcpy((void*)&CellularAddress, (void*)(&IfReqs[IdxReq].ifr_addr), sizeof(sockaddr_in));
+						FMemory::Memcpy((void*)&CellularAddress, (void*)(&TemporaryAddress), sizeof(sockaddr_in));
 					}
 					else if (OtherAddress.ss_family == AF_UNSPEC)
 					{
 						// First alternate found
-						FMemory::Memcpy((void*)&OtherAddress, (void*)(&IfReqs[IdxReq].ifr_addr), sizeof(sockaddr_in));
+						FMemory::Memcpy((void*)&OtherAddress, (void*)(&TemporaryAddress), sizeof(sockaddr_in));
 					}
 				}
 			}
@@ -244,3 +250,15 @@ TSharedRef<FInternetAddr> FSocketSubsystemAndroid::GetLocalHostAddr(FOutputDevic
 
 	return Addr;
 }
+
+/**
+ * Translates an ESocketAddressInfoFlags into a value usable by getaddrinfo
+ */
+int32 FSocketSubsystemAndroid::GetAddressInfoHintFlag(EAddressInfoFlags InFlags) const
+{
+	// On Android, you cannot explicitly use AI_ALL and AI_V4MAPPED
+	// However, if GetAddressInfo is passed with no hint flags, the query will execute with
+	// AI_V4MAPPED automatically (flag can only be set by the kernel)
+	EAddressInfoFlags ModifiedInFlags = (InFlags & ~EAddressInfoFlags::AllResultsWithMapping);
+	return FSocketSubsystemBSD::GetAddressInfoHintFlag(ModifiedInFlags);
+}
diff --git a/Engine/Source/Runtime/Sockets/Private/Android/SocketSubsystemAndroid.h b/Engine/Source/Runtime/Sockets/Private/Android/SocketSubsystemAndroid.h
--- a/Engine/Source/Runtime/Sockets/Private/Android/SocketSubsystemAndroid.h
+++ b/Engine/Source/Runtime/Sockets/Private/Android/SocketSubsystemAndroid.h
@@ -79,6 +79,8 @@ public:
 	virtual ESocketErrors GetHostByName(const ANSICHAR* HostName, FInternetAddr& OutAddr) override;
 	virtual ESocketErrors CreateAddressFromIP(const ANSICHAR* IPAddress, FInternetAddr& OutAddr) override;
 
+	virtual int32 GetAddressInfoHintFlag(EAddressInfoFlags InFlags) const override;
+
 	/**
 	 * Android platform specific look up to determine the host address
 	 * as many Android devices have multiple interfaces (wifi, cellular et al.)
