Fix SetMulticastInterface to check IPAddress protocol before setting (like all the others do [merge issue]).

#rnx

diff --git a/Engine/Source/Runtime/Sockets/Private/BSDSockets/SocketsBSD.cpp b/Engine/Source/Runtime/Sockets/Private/BSDSockets/SocketsBSD.cpp
--- a/Engine/Source/Runtime/Sockets/Private/BSDSockets/SocketsBSD.cpp
+++ b/Engine/Source/Runtime/Sockets/Private/BSDSockets/SocketsBSD.cpp
@@ -497,8 +497,11 @@ bool FSocketBSD::SetMulticastInterface(const FInternetAddr& InterfaceAddress)
 	const FInternetAddrBSD& BSDIFAddr = static_cast<const FInternetAddrBSD&>(InterfaceAddress);
 
 #if PLATFORM_HAS_BSD_IPV6_SOCKETS
-	uint32 InterfaceIndex = htonl(BSDIFAddr.GetScopeId());
-	return (setsockopt(Socket, IPPROTO_IPV6, IPV6_MULTICAST_IF, (char*)&InterfaceIndex, sizeof(InterfaceIndex)) == 0);
+	if (BSDIFAddr.GetProtocolFamily() == ESocketProtocolFamily::IPv6)
+	{
+		uint32 InterfaceIndex = htonl(BSDIFAddr.GetScopeId());
+		return (setsockopt(Socket, IPPROTO_IPV6, IPV6_MULTICAST_IF, (char*)&InterfaceIndex, sizeof(InterfaceIndex)) == 0);
+	}
 #endif
 
 	in_addr InterfaceAddr = ((sockaddr_in*)&(BSDIFAddr.Addr))->sin_addr;
