Fix IPv6 scope interfaces leaking when they shouldn't.

diff --git a/Engine/Source/Runtime/Sockets/Private/BSDIPv6Sockets/IPAddressBSDIPv6.cpp b/Engine/Source/Runtime/Sockets/Private/BSDIPv6Sockets/IPAddressBSDIPv6.cpp
--- a/Engine/Source/Runtime/Sockets/Private/BSDIPv6Sockets/IPAddressBSDIPv6.cpp
+++ b/Engine/Source/Runtime/Sockets/Private/BSDIPv6Sockets/IPAddressBSDIPv6.cpp
@@ -21,12 +21,6 @@ void FInternetAddrBSDIPv6::SetRawIp(const TArray<uint8>& RawAddr)
 			Addr.sin6_addr.s6_addr[i] = RawAddr[i];
 		}
 
-		// If this address has an interface, we'll copy it over too.
-		if (RawAddr.Num() == 20)
-		{
-			Addr.sin6_scope_id = (RawAddr[16] << 0) | (RawAddr[17] << 8) | (RawAddr[18] << 16) | (RawAddr[19] << 24);
-		}
-
 		Addr.sin6_family = AF_INET6;
 	}
 	else
@@ -44,15 +38,6 @@ TArray<uint8> FInternetAddrBSDIPv6::GetRawIp() const
 		RawAddressArray.Add(Addr.sin6_addr.s6_addr[i]);
 	}
 
-	// Copy over the interface if we have it explicitly set (anything other than 0)
-	if (Addr.sin6_scope_id != 0)
-	{
-		uint32 RawScopeId = Addr.sin6_scope_id;
-		RawAddressArray.Add((RawScopeId >> 0) & 0xFF);
-		RawAddressArray.Add((RawScopeId >> 8) & 0xFF);
-		RawAddressArray.Add((RawScopeId >> 16) & 0xFF);
-		RawAddressArray.Add((RawScopeId >> 24) & 0xFF);
-	}
 	return RawAddressArray;
 }
 
@@ -296,6 +281,16 @@ bool FInternetAddrBSDIPv6::IsValid() const
 	return memcmp(&Addr.sin6_addr, &Temp.Addr.sin6_addr, sizeof(in6_addr)) != 0;
 }
 
+uint32 FInternetAddrBSDIPv6::GetScopeId() const
+{
+	return ntohl(Addr.sin6_scope_id);
+}
+
+void FInternetAddrBSDIPv6::SetScopeId(uint32 NewScopeId)
+{
+	Addr.sin6_scope_id = htonl(NewScopeId);
+}
+
 TSharedRef<FInternetAddr> FInternetAddrBSDIPv6::Clone() const
 {
 
diff --git a/Engine/Source/Runtime/Sockets/Private/BSDIPv6Sockets/IPAddressBSDIPv6.h b/Engine/Source/Runtime/Sockets/Private/BSDIPv6Sockets/IPAddressBSDIPv6.h
--- a/Engine/Source/Runtime/Sockets/Private/BSDIPv6Sockets/IPAddressBSDIPv6.h
+++ b/Engine/Source/Runtime/Sockets/Private/BSDIPv6Sockets/IPAddressBSDIPv6.h
@@ -155,6 +155,23 @@ public:
 	}
 
 	virtual TSharedRef<FInternetAddr> Clone() const override;
+
+	/**
+	 * Sets the scope interface id of the currently held address if this address
+	 * is an IPv6 address internally. If it is not, no data will be assigned.
+	 * The NewScopeId must be in host byte order.
+	 *
+	 * @param NewScopeId the new scope interface id to set this address to
+	 */
+	virtual void SetScopeId(uint32 NewScopeId);
+
+	/**
+	 * Returns the IPv6 scope interface id of the currently held address
+	 * if the address is an IPv6 address.
+	 *
+	 * @return the scope interface id
+	 */
+	virtual uint32 GetScopeId() const;
 };
 
 
diff --git a/Engine/Source/Runtime/Sockets/Private/BSDIPv6Sockets/SocketSubsystemBSDIPv6.cpp b/Engine/Source/Runtime/Sockets/Private/BSDIPv6Sockets/SocketSubsystemBSDIPv6.cpp
--- a/Engine/Source/Runtime/Sockets/Private/BSDIPv6Sockets/SocketSubsystemBSDIPv6.cpp
+++ b/Engine/Source/Runtime/Sockets/Private/BSDIPv6Sockets/SocketSubsystemBSDIPv6.cpp
@@ -187,7 +187,9 @@ ESocketErrors FSocketSubsystemBSDIPv6::GetHostByName(const ANSICHAR* HostName, F
 
 	if (GAIResult.Results.Num() > 0)
 	{
-		OutAddr.SetRawIp(GAIResult.Results[0].Address->GetRawIp());
+		TSharedRef<FInternetAddrBSDIPv6> ResultAddr = StaticCastSharedRef<FInternetAddrBSDIPv6>(GAIResult.Results[0].Address);
+		OutAddr.SetRawIp(ResultAddr->GetRawIp());
+		static_cast<FInternetAddrBSDIPv6&>(OutAddr).SetScopeId(ResultAddr->GetScopeId());
 		return SE_NO_ERROR;
 	}
 
diff --git a/Engine/Source/Runtime/Sockets/Private/BSDSockets/IPAddressBSD.cpp b/Engine/Source/Runtime/Sockets/Private/BSDSockets/IPAddressBSD.cpp
--- a/Engine/Source/Runtime/Sockets/Private/BSDSockets/IPAddressBSD.cpp
+++ b/Engine/Source/Runtime/Sockets/Private/BSDSockets/IPAddressBSD.cpp
@@ -238,16 +238,6 @@ TArray<uint8> FInternetAddrBSD::GetRawIp() const
 		{
 			RawAddressArray.Add(IPv6Addr->sin6_addr.s6_addr[i]);
 		}
-
-		// Copy over the interface if we have it explicitly set (anything other than 0)
-		if (IPv6Addr->sin6_scope_id != 0)
-		{
-			uint32 RawScopeId = IPv6Addr->sin6_scope_id;
-			RawAddressArray.Add((RawScopeId >> 0) & 0xFF);
-			RawAddressArray.Add((RawScopeId >> 8) & 0xFF);
-			RawAddressArray.Add((RawScopeId >> 16) & 0xFF);
-			RawAddressArray.Add((RawScopeId >> 24) & 0xFF);
-		}
 	}
 #endif
 
@@ -271,12 +261,6 @@ void FInternetAddrBSD::SetRawIp(const TArray<uint8>& RawAddr)
 			IPv6Addr->sin6_addr.s6_addr[i] = RawAddr[i];
 		}
 
-		// If this address has an interface, we'll copy it over too.
-		if (RawAddr.Num() == 20)
-		{
-			IPv6Addr->sin6_scope_id = (RawAddr[16] << 0) | (RawAddr[17] << 8) | (RawAddr[18] << 16) | (RawAddr[19] << 24);
-		}
-
 		Addr.ss_family = AF_INET6;
 	}
 #endif
diff --git a/Engine/Source/Runtime/Sockets/Private/BSDSockets/SocketSubsystemBSD.cpp b/Engine/Source/Runtime/Sockets/Private/BSDSockets/SocketSubsystemBSD.cpp
--- a/Engine/Source/Runtime/Sockets/Private/BSDSockets/SocketSubsystemBSD.cpp
+++ b/Engine/Source/Runtime/Sockets/Private/BSDSockets/SocketSubsystemBSD.cpp
@@ -90,7 +90,9 @@ ESocketErrors FSocketSubsystemBSD::GetHostByName(const ANSICHAR* HostName, FInte
 
 	if (GAIResult.Results.Num() > 0)
 	{
-		OutAddr.SetRawIp(GAIResult.Results[0].Address->GetRawIp());
+		TSharedRef<FInternetAddrBSD> ResultAddr = StaticCastSharedRef<FInternetAddrBSD>(GAIResult.Results[0].Address);
+		OutAddr.SetRawIp(ResultAddr->GetRawIp());
+		static_cast<FInternetAddrBSD&>(OutAddr).SetScopeId(ResultAddr->GetScopeId());
 		return SE_NO_ERROR;
 	}
 
diff --git a/Engine/Source/Runtime/Sockets/Private/IOS/SocketSubsystemIOS.cpp b/Engine/Source/Runtime/Sockets/Private/IOS/SocketSubsystemIOS.cpp
--- a/Engine/Source/Runtime/Sockets/Private/IOS/SocketSubsystemIOS.cpp
+++ b/Engine/Source/Runtime/Sockets/Private/IOS/SocketSubsystemIOS.cpp
@@ -89,7 +89,9 @@ ESocketErrors FSocketSubsystemIOS::GetHostByName(const ANSICHAR* HostName, FInte
 
 	if (GAIResult.Results.Num() > 0)
 	{
-		OutAddr.SetRawIp(GAIResult.Results[0].Address->GetRawIp());
+		TSharedRef<FInternetAddrBSDIOS> ResultAddr = StaticCastSharedRef<FInternetAddrBSDIOS>(GAIResult.Results[0].Address);
+		OutAddr.SetRawIp(ResultAddr->GetRawIp());
+		static_cast<FInternetAddrBSDIOS&>(OutAddr).SetScopeId(ResultAddr->GetScopeId());
 		return SE_NO_ERROR;
 	}
 
