Updates to IPAddress structures that allow for querying protocol types. Adds errorcode tracking to GAI requests.

Updates the SocketSubsystem to better support multiple protocol types and handling platforms that are locked to a specific protocol as well. Also removes the enum usage in the sockets/subsystem classes which reverts and better implements 4504445.

diff --git a/Engine/Plugins/Online/OnlineSubsystemOculus/Source/Private/IPAddressOculus.h b/Engine/Plugins/Online/OnlineSubsystemOculus/Source/Private/IPAddressOculus.h
--- a/Engine/Plugins/Online/OnlineSubsystemOculus/Source/Private/IPAddressOculus.h
+++ b/Engine/Plugins/Online/OnlineSubsystemOculus/Source/Private/IPAddressOculus.h
@@ -6,6 +6,7 @@
 #include "IPAddress.h"
 #include "Algo/Reverse.h"
 #include "Engine/EngineBaseTypes.h"
+#include "OnlineSubsystemOculus.h"
 #include "OnlineSubsystemOculusTypes.h"
 #include "OnlineSubsystemOculusPackage.h"
 
@@ -201,11 +202,16 @@ public:
 		return !(FInternetAddrOculus::operator==(Other));
 	}
 
-	virtual uint32 GetTypeHash() override
+	virtual uint32 GetTypeHash() const override
 	{
 		return ::GetTypeHash((uint64)GetID());
 	}
 
+	virtual FName GetProtocolType() const override
+	{
+		return FNetworkProtocolTypes::Oculus;
+	}
+
 	/**
 	* Is this a well formed internet address
 	*
diff --git a/Engine/Plugins/Online/OnlineSubsystemOculus/Source/Private/OnlineSubsystemOculus.cpp b/Engine/Plugins/Online/OnlineSubsystemOculus/Source/Private/OnlineSubsystemOculus.cpp
--- a/Engine/Plugins/Online/OnlineSubsystemOculus/Source/Private/OnlineSubsystemOculus.cpp
+++ b/Engine/Plugins/Online/OnlineSubsystemOculus/Source/Private/OnlineSubsystemOculus.cpp
@@ -19,6 +19,11 @@
 OVRPL_PUBLIC_FUNCTION(void) ovr_ResetInitAndContext();
 #endif
 
+namespace FNetworkProtocolTypes
+{
+	const FName Oculus(TEXT("Oculus"));
+}
+
 IOnlineSessionPtr FOnlineSubsystemOculus::GetSessionInterface() const
 {
 	return SessionInterface;
diff --git a/Engine/Plugins/Online/OnlineSubsystemOculus/Source/Public/OnlineSubsystemOculus.h b/Engine/Plugins/Online/OnlineSubsystemOculus/Source/Public/OnlineSubsystemOculus.h
--- a/Engine/Plugins/Online/OnlineSubsystemOculus/Source/Public/OnlineSubsystemOculus.h
+++ b/Engine/Plugins/Online/OnlineSubsystemOculus/Source/Public/OnlineSubsystemOculus.h
@@ -129,4 +129,9 @@ private:
 	FOnlineMessageTaskManagerOculusPtr MessageTaskManager;
 };
 
+namespace FNetworkProtocolTypes
+{
+	ONLINESUBSYSTEMOCULUS_API extern const FName Oculus;
+}
+
 typedef TSharedPtr<FOnlineSubsystemOculus, ESPMode::ThreadSafe> FOnlineSubsystemOculusPtr;
diff --git a/Engine/Plugins/Online/OnlineSubsystemSteam/Source/Private/IPAddressSteam.h b/Engine/Plugins/Online/OnlineSubsystemSteam/Source/Private/IPAddressSteam.h
--- a/Engine/Plugins/Online/OnlineSubsystemSteam/Source/Private/IPAddressSteam.h
+++ b/Engine/Plugins/Online/OnlineSubsystemSteam/Source/Private/IPAddressSteam.h
@@ -4,6 +4,7 @@
 
 #include "CoreMinimal.h"
 #include "IPAddress.h"
+#include "OnlineSubsystemSteam.h"
 #include "OnlineSubsystemSteamTypes.h"
 #include "OnlineSubsystemSteamPackage.h"
 
@@ -163,19 +164,14 @@ public:
 		return !(FInternetAddrSteam::operator==(Other));
 	}
 
-	virtual uint32 GetTypeHash() override
-	{
-		return GetConstTypeHash();
-	}
-
-	uint32 GetConstTypeHash() const
+	virtual uint32 GetTypeHash() const override
 	{
 		return ::GetTypeHash(ToString(true));
 	}
 
 	friend uint32 GetTypeHash(const FInternetAddrSteam& A)
 	{
-		return A.GetConstTypeHash();
+		return A.GetTypeHash();
 	}
 
 	/**
@@ -188,6 +184,11 @@ public:
 		return SteamId.IsValid();
 	}
 
+	virtual FName GetProtocolType() const override
+	{
+		return FNetworkProtocolTypes::Steam;
+	}
+
 	virtual TSharedRef<FInternetAddr> Clone() const override
 	{
 		TSharedRef<FInternetAddrSteam> NewAddress = MakeShareable(new FInternetAddrSteam);
diff --git a/Engine/Plugins/Online/OnlineSubsystemSteam/Source/Private/OnlineSubsystemSteam.cpp b/Engine/Plugins/Online/OnlineSubsystemSteam/Source/Private/OnlineSubsystemSteam.cpp
--- a/Engine/Plugins/Online/OnlineSubsystemSteam/Source/Private/OnlineSubsystemSteam.cpp
+++ b/Engine/Plugins/Online/OnlineSubsystemSteam/Source/Private/OnlineSubsystemSteam.cpp
@@ -40,6 +40,11 @@
 #define UE4_PROJECT_STEAMSHIPPINGID 0
 #endif
 
+namespace FNetworkProtocolTypes
+{
+	const FName Steam(TEXT("Steam"));
+}
+
 #if !UE_BUILD_SHIPPING
 namespace OSSConsoleVariables
 {
diff --git a/Engine/Plugins/Online/OnlineSubsystemSteam/Source/Private/SocketSubsystemSteam.cpp b/Engine/Plugins/Online/OnlineSubsystemSteam/Source/Private/SocketSubsystemSteam.cpp
--- a/Engine/Plugins/Online/OnlineSubsystemSteam/Source/Private/SocketSubsystemSteam.cpp
+++ b/Engine/Plugins/Online/OnlineSubsystemSteam/Source/Private/SocketSubsystemSteam.cpp
@@ -165,7 +165,7 @@ void FSocketSubsystemSteam::Shutdown()
  *
  * @return the new socket or NULL if failed
  */
-FSocket* FSocketSubsystemSteam::CreateSocket(const FName& SocketType, const FString& SocketDescription, ESocketProtocolFamily ProtocolType)
+FSocket* FSocketSubsystemSteam::CreateSocket(const FName& SocketType, const FString& SocketDescription, const FName& ProtocolType)
 {
 	FSocket* NewSocket = nullptr;
 	if (SocketType == FName("SteamClientSocket"))
@@ -174,7 +174,7 @@ FSocket* FSocketSubsystemSteam::CreateSocket(const FName& SocketType, const FStr
 		if (SteamUserPtr != nullptr)
 		{
 			FUniqueNetIdSteam ClientId(SteamUserPtr->GetSteamID());
-			NewSocket = new FSocketSteam(SteamNetworking(), ClientId, SocketDescription, ProtocolType);
+			NewSocket = new FSocketSteam(SteamNetworking(), ClientId, SocketDescription, FNetworkProtocolTypes::Steam);
 
 			if (NewSocket)
 			{
@@ -191,12 +191,12 @@ FSocket* FSocketSubsystemSteam::CreateSocket(const FName& SocketType, const FStr
 			// If the GameServer connection hasn't been created yet, mark the socket as invalid for now
 			if (SessionInt->bSteamworksGameServerConnected && SessionInt->GameServerSteamId->IsValid() && SessionInt->bPolicyResponseReceived)
 			{
-				NewSocket = new FSocketSteam(SteamGameServerNetworking(), *SessionInt->GameServerSteamId, SocketDescription, ProtocolType);
+				NewSocket = new FSocketSteam(SteamGameServerNetworking(), *SessionInt->GameServerSteamId, SocketDescription, FNetworkProtocolTypes::Steam);
 			}
 			else
 			{
 				FUniqueNetIdSteam InvalidId(uint64(0));
-				NewSocket = new FSocketSteam(SteamGameServerNetworking(), InvalidId, SocketDescription, ProtocolType);
+				NewSocket = new FSocketSteam(SteamGameServerNetworking(), InvalidId, SocketDescription, FNetworkProtocolTypes::Steam);
 			}
 
 			if (NewSocket)
@@ -339,21 +339,82 @@ void FSocketSubsystemSteam::ConnectFailure(const FUniqueNetIdSteam& RemoteId)
  * @return the array of results from GetAddrInfo
  */
 FAddressInfoResult FSocketSubsystemSteam::GetAddressInfo(const TCHAR* HostName, const TCHAR* ServiceName,
-	EAddressInfoFlags QueryFlags, ESocketProtocolFamily ProtocolType,	ESocketType SocketType)
+	EAddressInfoFlags QueryFlags, const FName ProtocolTypeName, ESocketType SocketType)
 {
-	UE_LOG_ONLINE(Warning, TEXT("GetAddressInfo is not supported on Steam Sockets"));
-	return FAddressInfoResult(HostName, ServiceName);
+	int32 UnusedIndex;
+	FString RawAddress(HostName);
+	// Remove steam prefixes if they exist
+	RawAddress.RemoveFromStart(STEAM_URL_PREFIX);
+
+	if (!RawAddress.FindChar(TEXT(':'), UnusedIndex) && !RawAddress.FindChar(TEXT('.'), UnusedIndex))
+	{
+		FAddressInfoResult SteamResult(HostName, ServiceName);
+
+		// This is an Steam address
+		uint64 Id = FCString::Atoi64(*RawAddress);
+		if (Id != 0)
+		{
+			FString PortString(ServiceName);
+			SteamResult.ReturnCode = SE_NO_ERROR;
+			TSharedRef<FInternetAddrSteam> SteamIdAddress = StaticCastSharedRef<FInternetAddrSteam>(CreateInternetAddr());
+			SteamIdAddress->SteamId = FUniqueNetIdSteam(Id);
+			if (PortString.IsNumeric())
+			{
+				SteamIdAddress->SetPort(FCString::Atoi(*PortString));
+			}
+
+			SteamResult.Results.Add(FAddressInfoResultData(SteamIdAddress, 0, FNetworkProtocolTypes::Steam, SOCKTYPE_Unknown));
+			return SteamResult;
+		}
+		else
+		{
+			UE_LOG(LogSockets, Warning, TEXT("GetAddressInfo: Could not serialize %s into a SteamID, the ID was invalid."), *RawAddress);
+			SteamResult.ReturnCode = SE_HOST_NOT_FOUND;
+			return SteamResult;
+		}
+	}
+
+	return ISocketSubsystem::Get(PLATFORM_SOCKETSUBSYSTEM)->GetAddressInfo(HostName, ServiceName, QueryFlags, ProtocolTypeName, SocketType);
 }
 
 /**
- * Does a DNS look up of a host name
+ * Serializes a string that only contains an address.
+ *
+ * On Steam, this will take SteamIDs and serialize them into FInternetAddrSteam if it is determined
+ * the input string is an ID. Otherwise, this will give you back a FInternetAddrBSD.
  *
- * @param HostName the name of the host to look up
- * @param OutAddr the address to copy the IP address to
+ * This is a what you see is what you get, there is no DNS resolution of the input string,
+ * so only use this if you know you already have a valid ip address.
+ * Otherwise, feed the address to GetAddressInfo for guaranteed results.
+ *
+ * @param InAddress the address to serialize
+ *
+ * @return The FInternetAddr of the given string address. This will point to nullptr on failure.
  */
-ESocketErrors FSocketSubsystemSteam::GetHostByName(const ANSICHAR* HostName, FInternetAddr& OutAddr)
+TSharedPtr<FInternetAddr> FSocketSubsystemSteam::GetAddressFromString(const FString& InAddress)
 {
-	return SE_EADDRNOTAVAIL;
+	int32 UnusedIndex;
+	FString RawAddress = InAddress;
+	// Remove steam prefixes if they exist
+	RawAddress.RemoveFromStart(STEAM_URL_PREFIX);
+	if (!RawAddress.FindChar(TEXT(':'), UnusedIndex) && !RawAddress.FindChar(TEXT('.'), UnusedIndex))
+	{
+		// This is a Steam Address
+		uint64 Id = FCString::Atoi64(*RawAddress);
+		if (Id != 0)
+		{
+			TSharedRef<FInternetAddrSteam> ReturnAddress = StaticCastSharedRef<FInternetAddrSteam>(CreateInternetAddr());
+			ReturnAddress->SteamId = FUniqueNetIdSteam(Id);
+			return ReturnAddress;
+		}
+		else
+		{
+			UE_LOG(LogSockets, Warning, TEXT("Could not serialize %s into a SteamID, the ID was invalid."), *RawAddress);
+			return nullptr;
+		}
+	}
+	
+	return ISocketSubsystem::Get(PLATFORM_SOCKETSUBSYSTEM)->GetAddressFromString(InAddress);
 }
 
 /**
@@ -370,13 +431,10 @@ bool FSocketSubsystemSteam::GetHostName(FString& HostName)
 
 /**
  *	Create a proper FInternetAddr representation
- * @param Address host address
- * @param Port host port
  */
-TSharedRef<FInternetAddr> FSocketSubsystemSteam::CreateInternetAddr(uint32 Address, uint32 Port)
+TSharedRef<FInternetAddr> FSocketSubsystemSteam::CreateInternetAddr()
 {
-	FInternetAddrSteam* SteamAddr = new FInternetAddrSteam();
-	return MakeShareable(SteamAddr);
+	return MakeShareable(new FInternetAddrSteam());
 }
 
 /**
@@ -727,6 +785,12 @@ void FSocketSubsystemSteam::CleanupDeadConnections(bool bSkipLinger)
  */
 void FSocketSubsystemSteam::DumpSteamP2PSessionInfo(P2PSessionState_t& SessionInfo)
 {
+	FOnlineSubsystemSteam* SteamSubsystem = static_cast<FOnlineSubsystemSteam*>(IOnlineSubsystem::Get(STEAM_SUBSYSTEM));
+	if (SteamSubsystem == nullptr)
+	{
+		return;
+	}
+
 	TSharedRef<FInternetAddr> IpAddr = ISocketSubsystem::Get(PLATFORM_SOCKETSUBSYSTEM)->CreateInternetAddr(SessionInfo.m_nRemoteIP, SessionInfo.m_nRemotePort);
 	UE_LOG_ONLINE(Verbose, TEXT("- Detailed P2P session info:"));
 	UE_LOG_ONLINE(Verbose, TEXT("-- IPAddress: %s"), *IpAddr->ToString(true));
diff --git a/Engine/Plugins/Online/OnlineSubsystemSteam/Source/Private/SocketSubsystemSteam.h b/Engine/Plugins/Online/OnlineSubsystemSteam/Source/Private/SocketSubsystemSteam.h
--- a/Engine/Plugins/Online/OnlineSubsystemSteam/Source/Private/SocketSubsystemSteam.h
+++ b/Engine/Plugins/Online/OnlineSubsystemSteam/Source/Private/SocketSubsystemSteam.h
@@ -271,7 +271,7 @@ public:
 	 *
 	 * @return the new socket or NULL if failed
 	 */
-	virtual class FSocket* CreateSocket(const FName& SocketType, const FString& SocketDescription, ESocketProtocolFamily ProtocolType) override;
+	virtual class FSocket* CreateSocket(const FName& SocketType, const FString& SocketDescription, const FName& ProtocolType) override;
 
 	/**
 	 * Cleans up a socket class
@@ -297,16 +297,25 @@ public:
 	 */
 	virtual FAddressInfoResult GetAddressInfo(const TCHAR* HostName, const TCHAR* ServiceName = nullptr,
 		EAddressInfoFlags QueryFlags = EAddressInfoFlags::Default,
-		ESocketProtocolFamily ProtocolType = ESocketProtocolFamily::None,
+		const FName ProtocolTypeName = NAME_None,
 		ESocketType SocketType = ESocketType::SOCKTYPE_Unknown) override;
 
+
 	/**
-	 * Does a DNS look up of a host name
+	 * Serializes a string that only contains an address.
+	 * 
+	 * On Steam, this will take SteamIDs and serialize them into FInternetAddrSteam if it is determined
+	 * the input string is an ID. Otherwise, this will give you back a FInternetAddrBSD.
+	 *
+	 * This is a what you see is what you get, there is no DNS resolution of the input string, 
+	 * so only use this if you know you already have a valid ip address. 
+	 * Otherwise, feed the address to GetAddressInfo for guaranteed results.
+	 *
+	 * @param IPAddress the ip address to serialize
 	 *
-	 * @param HostName the name of the host to look up
-	 * @param OutAddr the address to copy the IP address to
+	 * @return The FInternetAddr of the given string address. This will point to nullptr on failure.
 	 */
-	virtual ESocketErrors GetHostByName(const ANSICHAR* HostName, FInternetAddr& OutAddr) override;
+	virtual TSharedPtr<FInternetAddr> GetAddressFromString(const FString& IPAddress) override;
 
 	/**
 	 * Some platforms require chat data (voice, text, etc.) to be placed into
@@ -338,10 +347,8 @@ public:
 
 	/**
 	 *	Create a proper FInternetAddr representation
-	 * @param Address host address
-	 * @param Port host port
 	 */
-	virtual TSharedRef<FInternetAddr> CreateInternetAddr(uint32 Address=0, uint32 Port=0) override;
+	virtual TSharedRef<FInternetAddr> CreateInternetAddr() override;
 
 	/**
 	 * @return Whether the machine has a properly configured network device or not
diff --git a/Engine/Plugins/Online/OnlineSubsystemSteam/Source/Private/SocketsSteam.h b/Engine/Plugins/Online/OnlineSubsystemSteam/Source/Private/SocketsSteam.h
--- a/Engine/Plugins/Online/OnlineSubsystemSteam/Source/Private/SocketsSteam.h
+++ b/Engine/Plugins/Online/OnlineSubsystemSteam/Source/Private/SocketsSteam.h
@@ -55,7 +55,7 @@ public:
 	 * @param InSocketDescription the debug description of the socket
 	 * @param InSocketProtocol the protocol used to create this socket.
 	 */
-	FSocketSteam(ISteamNetworking* InSteamNetworkingPtr, FUniqueNetIdSteam& InLocalSteamId, const FString& InSocketDescription, ESocketProtocolFamily InSocketProtocol) :
+	FSocketSteam(ISteamNetworking* InSteamNetworkingPtr, FUniqueNetIdSteam& InLocalSteamId, const FString& InSocketDescription, const FName& InSocketProtocol) :
 		FSocket(SOCKTYPE_Datagram, InSocketDescription, InSocketProtocol),
 		LocalSteamId(InLocalSteamId),
 		SteamChannel(0),
diff --git a/Engine/Plugins/Online/OnlineSubsystemSteam/Source/Public/OnlineSubsystemSteam.h b/Engine/Plugins/Online/OnlineSubsystemSteam/Source/Public/OnlineSubsystemSteam.h
--- a/Engine/Plugins/Online/OnlineSubsystemSteam/Source/Public/OnlineSubsystemSteam.h
+++ b/Engine/Plugins/Online/OnlineSubsystemSteam/Source/Public/OnlineSubsystemSteam.h
@@ -306,5 +306,10 @@ public:
 	}
 };
 
+namespace FNetworkProtocolTypes
+{
+	ONLINESUBSYSTEMSTEAM_API extern const FName Steam;
+}
+
 typedef TSharedPtr<FOnlineSubsystemSteam, ESPMode::ThreadSafe> FOnlineSubsystemSteamPtr;
 
diff --git a/Engine/Source/Runtime/Sockets/Private/Android/SocketSubsystemAndroid.cpp b/Engine/Source/Runtime/Sockets/Private/Android/SocketSubsystemAndroid.cpp
--- a/Engine/Source/Runtime/Sockets/Private/Android/SocketSubsystemAndroid.cpp
+++ b/Engine/Source/Runtime/Sockets/Private/Android/SocketSubsystemAndroid.cpp
@@ -101,34 +101,6 @@ const TCHAR* FSocketSubsystemAndroid::GetSocketAPIName() const
 	return TEXT("BSD_Android");
 }
 
-ESocketErrors FSocketSubsystemAndroid::GetHostByName(const ANSICHAR* HostName, FInternetAddr& OutAddr)
-{
-	FAddressInfoResult GAIResult = GetAddressInfo(ANSI_TO_TCHAR(HostName), nullptr, EAddressInfoFlags::Default);
-
-	if (GAIResult.Results.Num() > 0)
-	{
-		OutAddr.SetRawIp(GAIResult.Results[0].Address->GetRawIp());
-		return SE_NO_ERROR;
-	}
-
-	return SE_HOST_NOT_FOUND;
-}
-
-ESocketErrors FSocketSubsystemAndroid::CreateAddressFromIP(const ANSICHAR* IPAddress, FInternetAddr& OutAddr)
-{
-	FAddressInfoResult GAIResult = GetAddressInfo(ANSI_TO_TCHAR(IPAddress), nullptr,
-		EAddressInfoFlags::NoResolveHost | EAddressInfoFlags::OnlyUsableAddresses);
-
-	if (GAIResult.Results.Num() > 0)
-	{
-		OutAddr.SetRawIp(GAIResult.Results[0].Address->GetRawIp());
-		return SE_NO_ERROR;
-	}
-
-	return SE_HOST_NOT_FOUND;
-}
-
-
 TSharedRef<FInternetAddr> FSocketSubsystemAndroid::GetLocalHostAddr(FOutputDevice& Out, bool& bCanBindAll)
 {
 	// Get parent address first
@@ -143,20 +115,11 @@ TSharedRef<FInternetAddr> FSocketSubsystemAndroid::GetLocalHostAddr(FOutputDevic
 	//
 	// Also NOTE: Network can flip out behind applications back when connectivity changes. eg. Move out of wifi range.
 	// This seems to recover OK between matches as subsystems are reinited each session Host/Join.
-
-	uint32 ParentIp;
-	Addr->GetIp(ParentIp);
-	if (ParentIp != 0 && (ParentIp & 0xff000000) != 0x7f000000)
+	if (Addr->GetProtocolType() == FNetworkProtocolTypes::IPv4 && Addr->GetRawIp()[0] == 127)
 	{
 		return Addr;
 	}
 
-	// TODO: Android doesn't support ifaddrs either except in the Android OS 7.0+
-	// Which isn't super great either. So I guess this block is stuck the way it is.
-	// Other alternatives could be rtnetdevice but it's blocking 
-	// (sure you can make it non-blocking but you still have to wait for a recv otherwise you have no data)
-
-	// we need to go deeper...  (see http://man7.org/linux/man-pages/man7/netdevice.7.html)
 	int TempSocket = socket(PF_INET, SOCK_STREAM, 0);
 	if (TempSocket)
 	{
@@ -212,24 +175,24 @@ TSharedRef<FInternetAddr> FSocketSubsystemAndroid::GetLocalHostAddr(FOutputDevic
 				}
 			}
 
-			FInternetAddrBSD& NewAddrRef = static_cast<FInternetAddrBSD&>(Addr.Get());
+			TSharedRef<FInternetAddrBSD> NewAddrRef = StaticCastSharedRef<FInternetAddrBSD>(Addr);
 			// Prioritize results found
 			if (WifiAddress.ss_family != AF_UNSPEC)
 			{
 				// Prefer Wifi
-				NewAddrRef.SetIp(WifiAddress);
+				NewAddrRef->SetIp(WifiAddress);
 				UE_LOG(LogSockets, Log, TEXT("(%s) Wifi Adapter IP %s"), GetSocketAPIName(), *Addr->ToString(false));
 			}
 			else if (CellularAddress.ss_family != AF_UNSPEC)
 			{
 				// Then cellular
-				NewAddrRef.SetIp(CellularAddress);
+				NewAddrRef->SetIp(CellularAddress);
 				UE_LOG(LogSockets, Log, TEXT("(%s) Cellular Adapter IP %s"), GetSocketAPIName(), *Addr->ToString(false));
 			}
 			else if (OtherAddress.ss_family != AF_UNSPEC)
 			{
 				// Then whatever else was found
-				NewAddrRef.SetIp(OtherAddress);
+				NewAddrRef->SetIp(OtherAddress);
 				UE_LOG(LogSockets, Log, TEXT("(%s) Adapter IP %s"), GetSocketAPIName(), *Addr->ToString(false));
 			}
 			else
diff --git a/Engine/Source/Runtime/Sockets/Private/Android/SocketSubsystemAndroid.h b/Engine/Source/Runtime/Sockets/Private/Android/SocketSubsystemAndroid.h
--- a/Engine/Source/Runtime/Sockets/Private/Android/SocketSubsystemAndroid.h
+++ b/Engine/Source/Runtime/Sockets/Private/Android/SocketSubsystemAndroid.h
@@ -71,13 +71,8 @@ public:
 	virtual const TCHAR* GetSocketAPIName() const override;
 
 	/**
-	 * Does a DNS look up of a host name
-	 *
-	 * @param HostName the name of the host to look up
-	 * @param Addr the address to copy the IP address to
+	 * Translates an ESocketAddressInfoFlags into a value usable by getaddrinfo
 	 */
-	virtual ESocketErrors GetHostByName(const ANSICHAR* HostName, FInternetAddr& OutAddr) override;
-	virtual ESocketErrors CreateAddressFromIP(const ANSICHAR* IPAddress, FInternetAddr& OutAddr) override;
 
 	virtual int32 GetAddressInfoHintFlag(EAddressInfoFlags InFlags) const override;
 
@@ -92,4 +87,9 @@ public:
 	 * @return The local host address
 	 */
 	virtual TSharedRef<FInternetAddr> GetLocalHostAddr(FOutputDevice& Out, bool& bCanBindAll) override;
+
+	virtual FName GetDefaultSocketProtocolFamily() const override
+	{
+		return FNetworkProtocolTypes::IPv4;
+	}
 };
diff --git a/Engine/Source/Runtime/Sockets/Private/BSDIPv6Sockets/IPAddressBSDIPv6.cpp b/Engine/Source/Runtime/Sockets/Private/BSDIPv6Sockets/IPAddressBSDIPv6.cpp
--- a/Engine/Source/Runtime/Sockets/Private/BSDIPv6Sockets/IPAddressBSDIPv6.cpp
+++ b/Engine/Source/Runtime/Sockets/Private/BSDIPv6Sockets/IPAddressBSDIPv6.cpp
@@ -14,20 +14,25 @@ FInternetAddrBSDIPv6::FInternetAddrBSDIPv6()
 
 void FInternetAddrBSDIPv6::SetRawIp(const TArray<uint8>& RawAddr)
 {
-	if (RawAddr.Num() >= 16)
+	FMemory::Memzero(&Addr, sizeof(Addr));
+	if (RawAddr.Num() == 4)
+	{
+		Addr.sin6_addr.s6_addr[10] = 0xff;
+		Addr.sin6_addr.s6_addr[11] = 0xff;
+		Addr.sin6_addr.s6_addr[12] = RawAddr[0];
+		Addr.sin6_addr.s6_addr[13] = RawAddr[1];
+		Addr.sin6_addr.s6_addr[14] = RawAddr[2];
+		Addr.sin6_addr.s6_addr[15] = RawAddr[3];
+	}
+	else if (RawAddr.Num() == 16)
 	{
 		for (int i = 0; i < 16; ++i)
 		{
 			Addr.sin6_addr.s6_addr[i] = RawAddr[i];
 		}
-
-		Addr.sin6_family = AF_INET6;
-	}
-	else
-	{
-		FMemory::Memzero(&Addr, sizeof(Addr));
-		Addr.sin6_family = AF_INET6;
 	}
+
+	Addr.sin6_family = AF_INET6;
 }
 
 TArray<uint8> FInternetAddrBSDIPv6::GetRawIp() const
@@ -221,6 +226,16 @@ int32 FInternetAddrBSDIPv6::GetPort() const
 	return ntohs(Addr.sin6_port);
 }
 
+void FInternetAddrBSDIPv6::SetScopeId(uint32 NewScopeId)
+{
+	Addr.sin6_scope_id = htonl(NewScopeId);
+}
+
+uint32 FInternetAddrBSDIPv6::GetScopeId() const
+{
+	return ntohl(Addr.sin6_scope_id);
+}
+
 void FInternetAddrBSDIPv6::SetAnyAddress()
 {
 	SetIp(in6addr_any);
@@ -267,7 +282,7 @@ FString FInternetAddrBSDIPv6::ToString(bool bAppendPort) const
 	return Result;
 }
 
-uint32 FInternetAddrBSDIPv6::GetTypeHash()
+uint32 FInternetAddrBSDIPv6::GetTypeHash() const
 {
 	// @todo: Find a more efficient way to hash IPv6 addresses, if they are to be used with NetConnection's
 	//return Addr.sin_addr.s_addr + (Addr.sin_port * 23);
@@ -281,16 +296,6 @@ bool FInternetAddrBSDIPv6::IsValid() const
 	return memcmp(&Addr.sin6_addr, &Temp.Addr.sin6_addr, sizeof(in6_addr)) != 0;
 }
 
-uint32 FInternetAddrBSDIPv6::GetScopeId() const
-{
-	return ntohl(Addr.sin6_scope_id);
-}
-
-void FInternetAddrBSDIPv6::SetScopeId(uint32 NewScopeId)
-{
-	Addr.sin6_scope_id = htonl(NewScopeId);
-}
-
 TSharedRef<FInternetAddr> FInternetAddrBSDIPv6::Clone() const
 {
 
diff --git a/Engine/Source/Runtime/Sockets/Private/BSDIPv6Sockets/IPAddressBSDIPv6.h b/Engine/Source/Runtime/Sockets/Private/BSDIPv6Sockets/IPAddressBSDIPv6.h
--- a/Engine/Source/Runtime/Sockets/Private/BSDIPv6Sockets/IPAddressBSDIPv6.h
+++ b/Engine/Source/Runtime/Sockets/Private/BSDIPv6Sockets/IPAddressBSDIPv6.h
@@ -5,6 +5,7 @@
 #include "CoreMinimal.h"
 #include "BSDSockets/SocketSubsystemBSDPrivate.h"
 #include "SocketSubsystemBSDIPv6.h"
+#include "SocketTypes.h"
 #include "IPAddress.h"
 
 #if PLATFORM_HAS_BSD_IPV6_SOCKETS
@@ -112,6 +113,23 @@ public:
 	/** Returns the port number from this address in host byte order */
 	virtual int32 GetPort() const override;
 
+	/**
+	 * Sets the scope interface id of the currently held address if this address
+	 * is an IPv6 address internally. If it is not, no data will be assigned.
+	 * The NewScopeId must be in host byte order.
+	 *
+	 * @param NewScopeId the new scope interface id to set this address to
+	 */
+	virtual void SetScopeId(uint32 NewScopeId);
+
+	/**
+	 * Returns the IPv6 scope interface id of the currently held address
+	 * if the address is an IPv6 address.
+	 *
+	 * @return the scope interface id
+	 */
+	virtual uint32 GetScopeId() const;
+
 	/** Sets the address to be any address */
 	virtual void SetAnyAddress() override;
 
@@ -135,7 +153,7 @@ public:
 	 */
 	virtual bool operator==(const FInternetAddr& Other) const override;
 
-	virtual uint32 GetTypeHash() override;
+	virtual uint32 GetTypeHash() const override;
 
 	/**
 	 * Is this a well formed internet address, the only criteria being non-zero
@@ -156,22 +174,10 @@ public:
 
 	virtual TSharedRef<FInternetAddr> Clone() const override;
 
-	/**
-	 * Sets the scope interface id of the currently held address if this address
-	 * is an IPv6 address internally. If it is not, no data will be assigned.
-	 * The NewScopeId must be in host byte order.
-	 *
-	 * @param NewScopeId the new scope interface id to set this address to
-	 */
-	virtual void SetScopeId(uint32 NewScopeId);
-
-	/**
-	 * Returns the IPv6 scope interface id of the currently held address
-	 * if the address is an IPv6 address.
-	 *
-	 * @return the scope interface id
-	 */
-	virtual uint32 GetScopeId() const;
+	virtual FName GetProtocolType() const override
+	{
+		return FNetworkProtocolTypes::IPv6;
+	}
 };
 
 
diff --git a/Engine/Source/Runtime/Sockets/Private/BSDIPv6Sockets/SocketSubsystemBSDIPv6.cpp b/Engine/Source/Runtime/Sockets/Private/BSDIPv6Sockets/SocketSubsystemBSDIPv6.cpp
--- a/Engine/Source/Runtime/Sockets/Private/BSDIPv6Sockets/SocketSubsystemBSDIPv6.cpp
+++ b/Engine/Source/Runtime/Sockets/Private/BSDIPv6Sockets/SocketSubsystemBSDIPv6.cpp
@@ -8,7 +8,7 @@
 #include "BSDIPv6Sockets/IPAddressBSDIPv6.h"
 
 PRAGMA_DISABLE_DEPRECATION_WARNINGS
-class FSocketBSDIPv6* FSocketSubsystemBSDIPv6::InternalBSDSocketFactory(SOCKET Socket, ESocketType SocketType, const FString& SocketDescription, ESocketProtocolFamily SocketProtocol)
+class FSocketBSDIPv6* FSocketSubsystemBSDIPv6::InternalBSDSocketFactory(SOCKET Socket, ESocketType SocketType, const FString& SocketDescription, const FName& SocketProtocol)
 {
 	// return a new socket object
 	return new FSocketBSDIPv6(Socket, SocketType, SocketDescription, SocketProtocol, this);
@@ -46,28 +46,22 @@ ESocketErrors FSocketSubsystemBSDIPv6::TranslateGAIErrorCode(int32 Code) const
 	return SE_NO_ERROR;
 }
 
-FSocket* FSocketSubsystemBSDIPv6::CreateSocket(const FName& SocketType, const FString& SocketDescription, ESocketProtocolFamily ProtocolType)
+FSocket* FSocketSubsystemBSDIPv6::CreateSocket(const FName& SocketType, const FString& SocketDescription, const FName& ProtocolType)
 {
 	SOCKET Socket = INVALID_SOCKET;
 	FSocket* NewSocket = NULL;
 
-	// If we're passed None, initialize with IPv6 to maintain valid values.
-	if (ProtocolType == ESocketProtocolFamily::None)
-	{
-		ProtocolType = ESocketProtocolFamily::IPv6;
-	}
-
 	switch (SocketType.GetComparisonIndex())
 	{
 	case NAME_DGram:
 		// Creates a data gram (UDP) socket
 		Socket = socket(AF_INET6, SOCK_DGRAM, IPPROTO_UDP);
-		NewSocket = (Socket != INVALID_SOCKET) ? InternalBSDSocketFactory(Socket, SOCKTYPE_Datagram, SocketDescription, ProtocolType) : NULL;
+		NewSocket = (Socket != INVALID_SOCKET) ? InternalBSDSocketFactory(Socket, SOCKTYPE_Datagram, SocketDescription, FNetworkProtocolTypes::IPv6) : NULL;
 		break;
 	case NAME_Stream:
 		// Creates a stream (TCP) socket
 		Socket = socket(AF_INET6, SOCK_STREAM, IPPROTO_TCP);
-		NewSocket = (Socket != INVALID_SOCKET) ? InternalBSDSocketFactory(Socket, SOCKTYPE_Streaming, SocketDescription, ProtocolType) : NULL;
+		NewSocket = (Socket != INVALID_SOCKET) ? InternalBSDSocketFactory(Socket, SOCKTYPE_Streaming, SocketDescription, FNetworkProtocolTypes::IPv6) : NULL;
 		break;
 	}
 
@@ -90,7 +84,7 @@ void FSocketSubsystemBSDIPv6::DestroySocket(FSocket* Socket)
 }
 
 FAddressInfoResult FSocketSubsystemBSDIPv6::GetAddressInfo(const TCHAR* HostName, const TCHAR* ServiceName,
-	EAddressInfoFlags QueryFlags, ESocketProtocolFamily ProtocolType, ESocketType SocketType)
+	EAddressInfoFlags QueryFlags, const FName ProtocolTypeName, ESocketType SocketType)
 {
 	FAddressInfoResult AddrQueryResult = FAddressInfoResult(HostName, ServiceName);
 
@@ -106,8 +100,7 @@ FAddressInfoResult FSocketSubsystemBSDIPv6::GetAddressInfo(const TCHAR* HostName
 	addrinfo HintAddrInfo;
 	FMemory::Memzero(&HintAddrInfo, sizeof(HintAddrInfo));
 	HintAddrInfo.ai_family = AF_UNSPEC;
-	HintAddrInfo.ai_flags = GetAddressInfoHintFlag(QueryFlags);
-
+	
 	if (SocketType != ESocketType::SOCKTYPE_Unknown)
 	{
 		bool bIsUDP = (SocketType == ESocketType::SOCKTYPE_Datagram);
@@ -115,9 +108,16 @@ FAddressInfoResult FSocketSubsystemBSDIPv6::GetAddressInfo(const TCHAR* HostName
 		HintAddrInfo.ai_socktype = bIsUDP ? SOCK_DGRAM : SOCK_STREAM;
 	}
 
+	// Determine if we can save time with numericserv
+	if (ServiceName != nullptr && FString(ServiceName).IsNumeric())
+	{
+		QueryFlags |= EAddressInfoFlags::NoResolveService;
+	}
+	HintAddrInfo.ai_flags = GetAddressInfoHintFlag(QueryFlags);
+
 	int32 ErrorCode = getaddrinfo(TCHAR_TO_UTF8(HostName), TCHAR_TO_UTF8(ServiceName), &HintAddrInfo, &AddrInfo);
-	ESocketErrors SocketError = TranslateGAIErrorCode(ErrorCode);
-	if (SocketError == SE_NO_ERROR)
+	AddrQueryResult.ReturnCode = TranslateGAIErrorCode(ErrorCode);
+	if (AddrQueryResult.ReturnCode == SE_NO_ERROR)
 	{
 		addrinfo* AddrInfoHead = AddrInfo;
 		if (AddrInfo != nullptr && AddrInfo->ai_canonname != nullptr)
@@ -169,7 +169,7 @@ FAddressInfoResult FSocketSubsystemBSDIPv6::GetAddressInfo(const TCHAR* HostName
 
 				// Everything in this class is stored internally as IPv6
 				AddrQueryResult.Results.Add(FAddressInfoResultData(NewAddress, AddrInfo->ai_addrlen,
-					ESocketProtocolFamily::IPv6, ResultAddrConfiguration));
+					FNetworkProtocolTypes::IPv6, ResultAddrConfiguration));
 			}
 		}
 		freeaddrinfo(AddrInfoHead);
@@ -180,20 +180,36 @@ FAddressInfoResult FSocketSubsystemBSDIPv6::GetAddressInfo(const TCHAR* HostName
 	return AddrQueryResult;
 }
 
-ESocketErrors FSocketSubsystemBSDIPv6::GetHostByName(const ANSICHAR* HostName, FInternetAddr& OutAddr)
+TSharedPtr<FInternetAddr> FSocketSubsystemBSDIPv6::GetAddressFromString(const FString& InAddress)
 {
-	FAddressInfoResult GAIResult = GetAddressInfo(ANSI_TO_TCHAR(HostName), nullptr,
-		EAddressInfoFlags::AllResultsWithMapping | EAddressInfoFlags::OnlyUsableAddresses | EAddressInfoFlags::BindableAddress);
+	sockaddr_storage NetworkBuffer;
+	FMemory::Memzero(NetworkBuffer);
+
+	void* UniversalNetBufferPtr = nullptr;
+	int32 AddrFamily = AF_INET;
+
+	int32 ColonIndex;
+	if (InAddress.FindChar(TEXT(':'), ColonIndex))
+	{
+		AddrFamily = AF_INET6;
+		UniversalNetBufferPtr = &(((sockaddr_in6*)&NetworkBuffer)->sin6_addr);
+	}
+	else
+	{
+		UniversalNetBufferPtr = &(((sockaddr_in*)&NetworkBuffer)->sin_addr);
+	}
+
+	NetworkBuffer.ss_family = AddrFamily;
 
-	if (GAIResult.Results.Num() > 0)
+	if (inet_pton(AddrFamily, TCHAR_TO_ANSI(*InAddress), UniversalNetBufferPtr))
 	{
-		TSharedRef<FInternetAddrBSDIPv6> ResultAddr = StaticCastSharedRef<FInternetAddrBSDIPv6>(GAIResult.Results[0].Address);
-		OutAddr.SetRawIp(ResultAddr->GetRawIp());
-		static_cast<FInternetAddrBSDIPv6&>(OutAddr).SetScopeId(ResultAddr->GetScopeId());
-		return SE_NO_ERROR;
+		TSharedRef<FInternetAddrBSDIPv6> ReturnAddress = StaticCastSharedRef<FInternetAddrBSDIPv6>(CreateInternetAddr());
+		ReturnAddress->SetIp(NetworkBuffer);
+		return ReturnAddress;
 	}
 
-	return SE_HOST_NOT_FOUND;
+	UE_LOG(LogSockets, Warning, TEXT("Could not serialize %s, got error code %d"), *InAddress, GetLastErrorCode());
+	return nullptr;
 }
 
 bool FSocketSubsystemBSDIPv6::GetHostName(FString& HostName)
@@ -212,12 +228,9 @@ const TCHAR* FSocketSubsystemBSDIPv6::GetSocketAPIName() const
 	return TEXT("BSD IPv6");
 }
 
-TSharedRef<FInternetAddr> FSocketSubsystemBSDIPv6::CreateInternetAddr(uint32 Address, uint32 Port)
+TSharedRef<FInternetAddr> FSocketSubsystemBSDIPv6::CreateInternetAddr()
 {
-	TSharedRef<FInternetAddr> Result = MakeShareable(new FInternetAddrBSDIPv6);
-	Result->SetIp(Address);
-	Result->SetPort(Port);
-	return Result;
+	return MakeShareable(new FInternetAddrBSDIPv6);
 }
 
 bool FSocketSubsystemBSDIPv6::IsSocketWaitSupported() const
diff --git a/Engine/Source/Runtime/Sockets/Private/BSDIPv6Sockets/SocketSubsystemBSDIPv6.h b/Engine/Source/Runtime/Sockets/Private/BSDIPv6Sockets/SocketSubsystemBSDIPv6.h
--- a/Engine/Source/Runtime/Sockets/Private/BSDIPv6Sockets/SocketSubsystemBSDIPv6.h
+++ b/Engine/Source/Runtime/Sockets/Private/BSDIPv6Sockets/SocketSubsystemBSDIPv6.h
@@ -7,6 +7,7 @@
 #include "BSDSockets/SocketSubsystemBSDPrivate.h"
 #include "BSDIPv6Sockets/SocketsBSDIPv6.h"
 #include "IPAddress.h"
+#include "SocketTypes.h"
 #include "SocketSubsystemPackage.h"
 
 #if PLATFORM_HAS_BSD_IPV6_SOCKETS
@@ -21,20 +22,25 @@ public:
 
 	//~ Begin ISocketSubsystem Interface
 
-	virtual TSharedRef<FInternetAddr> CreateInternetAddr( uint32 Address = 0, uint32 Port = 0 ) override;
+	virtual TSharedRef<FInternetAddr> CreateInternetAddr() override;
 
-	virtual class FSocket* CreateSocket(const FName& SocketType, const FString& SocketDescription, ESocketProtocolFamily ProtocolType) override;
+	virtual class FSocket* CreateSocket(const FName& SocketType, const FString& SocketDescription, bool bForceUDP = false) override
+	{
+		return CreateSocket(SocketType, SocketDescription, FNetworkProtocolTypes::IPv6);
+	}
+
+	virtual class FSocket* CreateSocket(const FName& SocketType, const FString& SocketDescription, const FName& ProtocolType) override;
 
 	virtual FResolveInfoCached* CreateResolveInfoCached(TSharedPtr<FInternetAddr> Addr) const override;
 
-	virtual void DestroySocket( class FSocket* Socket ) override;
+	virtual void DestroySocket(class FSocket* Socket) override;
 
 	virtual FAddressInfoResult GetAddressInfo(const TCHAR* HostName, const TCHAR* ServiceName = nullptr,
 		EAddressInfoFlags QueryFlags = EAddressInfoFlags::Default,
-		ESocketProtocolFamily ProtocolType = ESocketProtocolFamily::None,
+		const FName ProtocolTypeName = NAME_None,
 		ESocketType SocketType = ESocketType::SOCKTYPE_Unknown) override;
 
-	virtual ESocketErrors GetHostByName( const ANSICHAR* HostName, FInternetAddr& OutAddr ) override;
+	virtual TSharedPtr<FInternetAddr> GetAddressFromString(const FString& InAddress) override;
 
 	virtual bool GetHostName( FString& HostName ) override;
 
@@ -73,7 +79,7 @@ protected:
 	/**
 	 * Allows a subsystem subclass to create a FSocketBSD sub class.
 	 */
-	virtual class FSocketBSDIPv6* InternalBSDSocketFactory( SOCKET Socket, ESocketType SocketType, const FString& SocketDescription, ESocketProtocolFamily SocketProtocol );
+	virtual class FSocketBSDIPv6* InternalBSDSocketFactory(SOCKET Socket, ESocketType SocketType, const FString& SocketDescription, const FName& SocketProtocol);
 
 
 	/**
@@ -82,7 +88,7 @@ protected:
 	UE_DEPRECATED(4.22, "To support multiple stack types, move to the constructor that allows for specifying the protocol stack to initialize the socket on.")
 	virtual class FSocketBSDIPv6* InternalBSDSocketFactory(SOCKET Socket, ESocketType SocketType, const FString& SocketDescription)
 	{
-		return InternalBSDSocketFactory(Socket, SocketType, SocketDescription, ESocketProtocolFamily::IPv6);
+		return InternalBSDSocketFactory(Socket, SocketType, SocketDescription, FNetworkProtocolTypes::IPv6);
 	}
 
 	/**
diff --git a/Engine/Source/Runtime/Sockets/Private/BSDIPv6Sockets/SocketsBSDIPv6.h b/Engine/Source/Runtime/Sockets/Private/BSDIPv6Sockets/SocketsBSDIPv6.h
--- a/Engine/Source/Runtime/Sockets/Private/BSDIPv6Sockets/SocketsBSDIPv6.h
+++ b/Engine/Source/Runtime/Sockets/Private/BSDIPv6Sockets/SocketsBSDIPv6.h
@@ -50,7 +50,7 @@ public:
 	 */
 	UE_DEPRECATED(4.22, "Use the socket constructor with protocol specification.")
 	FSocketBSDIPv6(SOCKET InSocket, ESocketType InSocketType, const FString& InSocketDescription, ISocketSubsystem* InSubsystem)
-		: FSocket(InSocketType, InSocketDescription, ESocketProtocolFamily::IPv6)
+		: FSocket(InSocketType, InSocketDescription, FNetworkProtocolTypes::IPv6)
 		, Socket(InSocket)
 		, SocketSubsystem(InSubsystem)
 	{ }
@@ -64,7 +64,7 @@ public:
 	 * @param InSocketProtocol the protocol used to create this socket.
 	 * @param InSubsystem the subsystem that created this socket
 	 */
-	FSocketBSDIPv6(SOCKET InSocket, ESocketType InSocketType, const FString& InSocketDescription, ESocketProtocolFamily InSocketProtocol, ISocketSubsystem* InSubsystem) 
+	FSocketBSDIPv6(SOCKET InSocket, ESocketType InSocketType, const FString& InSocketDescription, const FName& InSocketProtocol, ISocketSubsystem* InSubsystem) 
 		: FSocket(InSocketType, InSocketDescription, InSocketProtocol)
 		, Socket(InSocket)
 		, SocketSubsystem(InSubsystem)
diff --git a/Engine/Source/Runtime/Sockets/Private/BSDSockets/IPAddressBSD.cpp b/Engine/Source/Runtime/Sockets/Private/BSDSockets/IPAddressBSD.cpp
--- a/Engine/Source/Runtime/Sockets/Private/BSDSockets/IPAddressBSD.cpp
+++ b/Engine/Source/Runtime/Sockets/Private/BSDSockets/IPAddressBSD.cpp
@@ -5,9 +5,6 @@
 
 #if PLATFORM_HAS_BSD_SOCKETS || PLATFORM_HAS_BSD_IPV6_SOCKETS
 
-// Hardcoded address that the messagebus uses. This is a hack.
-#define IPV4_MESSAGEBUS_ADDRESS_HACK ((230 << 24) | (0 << 16) | (0 << 8) | (1 << 0))
-
 #if PLATFORM_HAS_BSD_IPV6_SOCKETS
 void MapIPv4ToIPv6(const uint32& InAddress, in6_addr& OutStructure)
 {
@@ -165,13 +162,17 @@ void FInternetAddrBSD::SetIp(const TCHAR* InAddr, bool& bIsValid)
 	AddressString.RemoveFromStart("[");
 	AddressString.RemoveFromEnd("]");
 	
-	const auto InAddrAnsi = StringCast<ANSICHAR>(*AddressString);
 	if (SocketSubsystem != nullptr)
 	{
-		bIsValid = (SocketSubsystem->CreateAddressFromIP(InAddrAnsi.Get(), *this) == SE_NO_ERROR);
-		if (bHasPort && bIsValid)
+		TSharedPtr<FInternetAddr> NewAddr = SocketSubsystem->GetAddressFromString(AddressString);
+		if (NewAddr.IsValid())
 		{
-			SetPort(FCString::Atoi(*Port));
+			bIsValid = true;
+			SetRawIp(NewAddr->GetRawIp());
+			if (bHasPort)
+			{
+				SetPort(FCString::Atoi(*Port));
+			}
 		}
 	}
 	else
@@ -183,13 +184,13 @@ void FInternetAddrBSD::SetIp(const TCHAR* InAddr, bool& bIsValid)
 void FInternetAddrBSD::SetIp(uint32 InAddr)
 {
 #if PLATFORM_HAS_BSD_IPV6_SOCKETS
-	if (SocketSubsystem && SocketSubsystem->GetDefaultSocketProtocolFamily() == ESocketProtocolFamily::IPv6)
+	if (SocketSubsystem && SocketSubsystem->GetDefaultSocketProtocolFamily() == FNetworkProtocolTypes::IPv6)
 	{
 		if (InAddr == 0)
 		{
 			SetAnyIPv6Address();
 		}
-		else if (InAddr == INADDR_BROADCAST || InAddr == IPV4_MESSAGEBUS_ADDRESS_HACK)
+		else if (InAddr == INADDR_BROADCAST)
 		{
 			SetIPv6BroadcastAddress();
 		}
@@ -272,7 +273,7 @@ void FInternetAddrBSD::SetRawIp(const TArray<uint8>& RawAddr)
 
 void FInternetAddrBSD::GetIp(uint32& OutAddr) const
 {
-	if (GetProtocolFamily() != ESocketProtocolFamily::IPv4)
+	if (GetProtocolType() == FNetworkProtocolTypes::IPv6)
 	{
 		OutAddr = 0;
 
@@ -296,7 +297,7 @@ void FInternetAddrBSD::GetIp(uint32& OutAddr) const
 void FInternetAddrBSD::SetPort(int32 InPort)
 {
 #if PLATFORM_HAS_BSD_IPV6_SOCKETS
-	if (GetProtocolFamily() == ESocketProtocolFamily::IPv6)
+	if (GetProtocolType() == FNetworkProtocolTypes::IPv6)
 	{
 		((sockaddr_in6*)&Addr)->sin6_port = htons(InPort);
 		return;
@@ -309,7 +310,7 @@ void FInternetAddrBSD::SetPort(int32 InPort)
 int32 FInternetAddrBSD::GetPort() const
 {
 #if PLATFORM_HAS_BSD_IPV6_SOCKETS
-	if (GetProtocolFamily() == ESocketProtocolFamily::IPv6)
+	if (GetProtocolType() == FNetworkProtocolTypes::IPv6)
 	{
 		return ntohs(((sockaddr_in6*)&Addr)->sin6_port);
 	}
@@ -322,7 +323,14 @@ void FInternetAddrBSD::SetAnyAddress()
 {
 	if (SocketSubsystem != nullptr)
 	{
-		SetAnyAddress(SocketSubsystem->GetDefaultSocketProtocolFamily());
+		if (SocketSubsystem->GetDefaultSocketProtocolFamily() == FNetworkProtocolTypes::IPv6)
+		{
+			SetAnyIPv6Address();
+		}
+		else
+		{
+			SetAnyIPv4Address();
+		}
 	}
 	else
 	{
@@ -349,7 +357,14 @@ void FInternetAddrBSD::SetBroadcastAddress()
 {
 	if (SocketSubsystem)
 	{
-		SetBroadcastAddress(SocketSubsystem->GetDefaultSocketProtocolFamily());
+		if (SocketSubsystem->GetDefaultSocketProtocolFamily() == FNetworkProtocolTypes::IPv6)
+		{
+			SetIPv6BroadcastAddress();
+		}
+		else
+		{
+			SetIPv4BroadcastAddress();
+		}
 	}
 	else
 	{
@@ -384,7 +399,14 @@ void FInternetAddrBSD::SetLoopbackAddress()
 {
 	if (SocketSubsystem)
 	{
-		SetLoopbackAddress(SocketSubsystem->GetDefaultSocketProtocolFamily());
+		if (SocketSubsystem->GetDefaultSocketProtocolFamily() == FNetworkProtocolTypes::IPv6)
+		{
+			SetIPv6LoopbackAddress();
+		}
+		else
+		{
+			SetIPv4LoopbackAddress();
+		}
 	}
 	else
 	{
@@ -414,7 +436,7 @@ FString FInternetAddrBSD::ToString(bool bAppendPort) const
 	char IPStr[NI_MAXHOST];
 	if (getnameinfo((const sockaddr*)&Addr, GetStorageSize(), IPStr, NI_MAXHOST, nullptr, 0, NI_NUMERICHOST) == 0)
 	{
-		if (GetProtocolFamily() == ESocketProtocolFamily::IPv6)
+		if (GetProtocolType() == FNetworkProtocolTypes::IPv6)
 		{
 			FString IPv6Str(ANSI_TO_TCHAR(IPStr));
 			// Remove the scope interface if it exists.
@@ -445,10 +467,10 @@ FString FInternetAddrBSD::ToString(bool bAppendPort) const
 bool FInternetAddrBSD::operator==(const FInternetAddr& Other) const
 {
 	const FInternetAddrBSD& OtherBSD = static_cast<const FInternetAddrBSD&>(Other);
-	ESocketProtocolFamily CurrentFamily = GetProtocolFamily();
+	FName CurrentFamily = GetProtocolType();
 
 	// Check if the addr families match
-	if (OtherBSD.GetProtocolFamily() != CurrentFamily)
+	if (OtherBSD.GetProtocolType() != CurrentFamily)
 	{
 		return false;
 	}
@@ -460,15 +482,15 @@ bool FInternetAddrBSD::operator==(const FInternetAddr& Other) const
 	}
 
 #if PLATFORM_HAS_BSD_IPV6_SOCKETS
-	if (CurrentFamily == ESocketProtocolFamily::IPv6)
-		{
+	if (CurrentFamily == FNetworkProtocolTypes::IPv6)
+	{
 		const sockaddr_in6* OtherBSDAddr = (sockaddr_in6*)&(OtherBSD.Addr);
 		const sockaddr_in6* ThisBSDAddr = ((sockaddr_in6*)&Addr);
 		return memcmp(&(ThisBSDAddr->sin6_addr), &(OtherBSDAddr->sin6_addr), sizeof(in6_addr)) == 0;
 	}
 #endif
 
-	if (CurrentFamily == ESocketProtocolFamily::IPv4)
+	if (CurrentFamily == FNetworkProtocolTypes::IPv4)
 	{
 		const sockaddr_in* OtherBSDAddr = (sockaddr_in*)&(OtherBSD.Addr);
 		const sockaddr_in* ThisBSDAddr = ((sockaddr_in*)&Addr);
@@ -480,15 +502,15 @@ bool FInternetAddrBSD::operator==(const FInternetAddr& Other) const
 
 bool FInternetAddrBSD::IsValid() const
 {
-	ESocketProtocolFamily CurrentFamily = GetProtocolFamily();
+	FName CurrentFamily = GetProtocolType();
 
-	if (CurrentFamily == ESocketProtocolFamily::IPv4)
+	if (CurrentFamily == FNetworkProtocolTypes::IPv4)
 	{
 		return ((sockaddr_in*)&Addr)->sin_addr.s_addr != 0;
 	}
 
 #if PLATFORM_HAS_BSD_IPV6_SOCKETS
-	if (CurrentFamily == ESocketProtocolFamily::IPv6)
+	if (CurrentFamily == FNetworkProtocolTypes::IPv6)
 	{
 		in6_addr EmptyAddr;
 		FMemory::Memzero(EmptyAddr);
@@ -508,24 +530,24 @@ TSharedRef<FInternetAddr> FInternetAddrBSD::Clone() const
 	return NewAddress;
 }
 
-ESocketProtocolFamily FInternetAddrBSD::GetProtocolFamily() const
+FName FInternetAddrBSD::GetProtocolType() const
 {
 	switch (Addr.ss_family)
 	{
 		case AF_INET:
-			return ESocketProtocolFamily::IPv4;
+			return FNetworkProtocolTypes::IPv4;
 		break;
 		case AF_INET6:
-			return ESocketProtocolFamily::IPv6;
+			return FNetworkProtocolTypes::IPv6;
 		break;
 		default:
-			return ESocketProtocolFamily::None;
+			return NAME_None;
 	}
 }
 
 SOCKLEN FInternetAddrBSD::GetStorageSize() const
 {
-	if (GetProtocolFamily() == ESocketProtocolFamily::IPv4)
+	if (GetProtocolType() == FNetworkProtocolTypes::IPv4)
 	{
 		return sizeof(sockaddr_in);
 	}
@@ -538,17 +560,15 @@ SOCKLEN FInternetAddrBSD::GetStorageSize() const
 	return sizeof(sockaddr_storage);
 }
 
-uint32 FInternetAddrBSD::GetTypeHash()
+uint32 FInternetAddrBSD::GetTypeHash() const
 {
-	ESocketProtocolFamily CurrentFamily = GetProtocolFamily();
+	FName CurrentFamily = GetProtocolType();
 
-	if (CurrentFamily == ESocketProtocolFamily::IPv4)
+	if (CurrentFamily == FNetworkProtocolTypes::IPv4)
 	{
-		uint32 NumericAddress;
-		GetIp(NumericAddress);
-		return NumericAddress + (GetPort() * 23);
+		return ntohl(((sockaddr_in*)&Addr)->sin_addr.s_addr) + (GetPort() * 23);
 	}
-	else if (CurrentFamily == ESocketProtocolFamily::IPv6)
+	else if (CurrentFamily == FNetworkProtocolTypes::IPv6)
 	{
 		return ::GetTypeHash(*ToString(true));
 	}
diff --git a/Engine/Source/Runtime/Sockets/Private/BSDSockets/IPAddressBSD.h b/Engine/Source/Runtime/Sockets/Private/BSDSockets/IPAddressBSD.h
--- a/Engine/Source/Runtime/Sockets/Private/BSDSockets/IPAddressBSD.h
+++ b/Engine/Source/Runtime/Sockets/Private/BSDSockets/IPAddressBSD.h
@@ -5,6 +5,8 @@
 #include "CoreMinimal.h"
 #include "BSDSockets/SocketSubsystemBSDPrivate.h"
 #include "IPAddress.h"
+#include "SocketTypes.h"
+#include "SocketSubsystemPackage.h"
 
 #if PLATFORM_HAS_BSD_SOCKETS || PLATFORM_HAS_BSD_IPV6_SOCKETS
 
@@ -23,12 +25,17 @@ protected:
 	/** The Subsystem that created this address */
 	FSocketSubsystemBSD* SocketSubsystem;
 
+	void Clear();
+	void ResetScopeId();
+
+PACKAGE_SCOPE:
+
 	/**
 	 * Sets the ip address using a network byte order ipv4 address
 	 *
 	 * @param IPv4Addr the new ip address to use
 	 */
-	void SetIp(const in_addr& IPv4Addr)
+	virtual void SetIp(const in_addr& IPv4Addr)
 	{
 		((sockaddr_in*)&Addr)->sin_addr = IPv4Addr;
 		Addr.ss_family = AF_INET;
@@ -40,16 +47,13 @@ protected:
 	 *
 	 * @param IpAddr the new ip address to use
 	 */
-	void SetIp(const in6_addr& IpAddr)
+	virtual void SetIp(const in6_addr& IPv6Addr)
 	{
-		((sockaddr_in6*)&Addr)->sin6_addr = IpAddr;
+		((sockaddr_in6*)&Addr)->sin6_addr = IPv6Addr;
 		Addr.ss_family = AF_INET6;
 	}
 #endif
 
-	void Clear();
-	void ResetScopeId();
-
 public:
 	/**
 	 * Constructor. Sets address to default state
@@ -134,7 +138,7 @@ public:
 	 */
 	void GetIp(in6_addr& OutAddr) const
 	{
-		if (GetProtocolFamily() != ESocketProtocolFamily::IPv6)
+		if (GetProtocolType() != FNetworkProtocolTypes::IPv6)
 		{
 			return;
 		}
@@ -150,7 +154,7 @@ public:
 	 */
 	void GetIp(in_addr& OutAddr) const
 	{
-		if (GetProtocolFamily() != ESocketProtocolFamily::IPv4)
+		if (GetProtocolType() != FNetworkProtocolTypes::IPv4)
 		{
 			return;
 		}
@@ -185,24 +189,6 @@ public:
 	 */
 	virtual void SetAnyAddress() override;
 
-	/**
-	 * Sets the address structure to be bindable to any IP address
-	 * for the protocol family specified.
-	 *
-	 * @param ForProtocol the protocol to set the address for
-	 */
-	virtual void SetAnyAddress(ESocketProtocolFamily ForProtocol)
-	{
-		if (ForProtocol == ESocketProtocolFamily::IPv4)
-		{
-			SetAnyIPv4Address();
-		}
-		else
-		{
-			SetAnyIPv6Address();
-		}
-	}
-
 	/** Explicit set to any IPv4 address */
 	virtual void SetAnyIPv4Address();
 
@@ -218,24 +204,6 @@ public:
 	 */
 	virtual void SetBroadcastAddress() override;
 
-	/**
-	 * Sets the address structure to be bound to the multicast ip address
-	 * for the protocol family specified.
-	 *
-	 * @param ForProtocol the protocol to set the address for
-	 */
-	virtual void SetBroadcastAddress(ESocketProtocolFamily ForProtocol)
-	{
-		if (ForProtocol == ESocketProtocolFamily::IPv4)
-		{
-			SetIPv4BroadcastAddress();
-		}
-		else
-		{
-			SetIPv6BroadcastAddress();
-		}
-	}
-
 	/** Explicit set to multicast IPv4 address */
 	virtual void SetIPv4BroadcastAddress();
 
@@ -251,24 +219,6 @@ public:
 	 */
 	virtual void SetLoopbackAddress() override;
 
-	/**
-	 * Sets the address structure to be bound to the loopback ip address
-	 * for the protocol family specified.
-	 *
-	 * @param ForProtocol the protocol to set the address for
-	 */
-	virtual void SetLoopbackAddress(ESocketProtocolFamily ForProtocol)
-	{
-		if (ForProtocol == ESocketProtocolFamily::IPv4)
-		{
-			SetIPv4LoopbackAddress();
-		}
-		else
-		{
-			SetIPv6LoopbackAddress();
-		}
-	}
-
 	/** Explicit set to loopback IPv4 address */
 	virtual void SetIPv4LoopbackAddress();
 
@@ -306,9 +256,9 @@ public:
 	/**
 	 * Returns the protocol family of the address data currently stored in this struct
 	 *
-	 * @return They type of the address
+	 * @return The type of the address. If it's not known or overridden, the address passes None.
 	 */
-	virtual ESocketProtocolFamily GetProtocolFamily() const;
+	virtual FName GetProtocolType() const override;
 
 	/**
 	 * Returns the size of the amount of data that is being used
@@ -335,7 +285,7 @@ public:
 	 */
 	virtual uint32 GetScopeId() const;
 	
-	virtual uint32 GetTypeHash() override;
+	virtual uint32 GetTypeHash() const override;
 
 	friend class FSocketBSD;
 };
diff --git a/Engine/Source/Runtime/Sockets/Private/BSDSockets/SocketSubsystemBSD.cpp b/Engine/Source/Runtime/Sockets/Private/BSDSockets/SocketSubsystemBSD.cpp
--- a/Engine/Source/Runtime/Sockets/Private/BSDSockets/SocketSubsystemBSD.cpp
+++ b/Engine/Source/Runtime/Sockets/Private/BSDSockets/SocketSubsystemBSD.cpp
@@ -4,39 +4,37 @@
 
 #if PLATFORM_HAS_BSD_SOCKETS || PLATFORM_HAS_BSD_IPV6_SOCKETS
 
-#include "IPAddress.h"
 #include "BSDSockets/IPAddressBSD.h"
 #include "BSDSockets/SocketsBSD.h"
+#include "Misc/CString.h"
 #include <errno.h>
 
-FSocketBSD* FSocketSubsystemBSD::InternalBSDSocketFactory(SOCKET Socket, ESocketType SocketType, const FString& SocketDescription, ESocketProtocolFamily SocketProtocol)
+FSocketBSD* FSocketSubsystemBSD::InternalBSDSocketFactory(SOCKET Socket, ESocketType SocketType, const FString& SocketDescription, const FName& SocketProtocol)
 {
 	// return a new socket object
 	return new FSocketBSD(Socket, SocketType, SocketDescription, SocketProtocol, this);
 }
 
-FSocket* FSocketSubsystemBSD::CreateSocket(const FName& SocketType, const FString& SocketDescription, bool bForceUDP)
-{
-	return CreateSocket(SocketType, SocketDescription, GetDefaultSocketProtocolFamily());
-}
-
-FSocket* FSocketSubsystemBSD::CreateSocket(const FName& SocketType, const FString& SocketDescription, ESocketProtocolFamily ProtocolType)
+FSocket* FSocketSubsystemBSD::CreateSocket(const FName& SocketType, const FString& SocketDescription, const FName& ProtocolType)
 {
 	SOCKET Socket = INVALID_SOCKET;
 	FSocket* NewSocket = nullptr;
 	int PlatformSpecificTypeFlags = 0;
 
+	FName SocketProtocolType = ProtocolType;
+
 	// For platforms that have two subsystems (ex: Steam) but don't explicitly inherit from SocketSubsystemBSD
 	// so they don't know which protocol to end up using and pass in None.
 	// This is invalid, so we need to attempt to still resolve it.
-	if (ProtocolType == ESocketProtocolFamily::None)
+	if (SocketProtocolType != FNetworkProtocolTypes::IPv4 || SocketProtocolType != FNetworkProtocolTypes::IPv6)
 	{
-		ProtocolType = GetDefaultSocketProtocolFamily();
+		SocketProtocolType = GetDefaultSocketProtocolFamily();
 	}
 
 	// Don't support any other protocol families.
-	if (ProtocolType != ESocketProtocolFamily::IPv4 && ProtocolType != ESocketProtocolFamily::IPv6)
+	if (SocketProtocolType != FNetworkProtocolTypes::IPv4 && SocketProtocolType != FNetworkProtocolTypes::IPv6)
 	{
+		UE_LOG(LogSockets, Warning, TEXT("Provided socket protocol type is unsupported! Returning null socket"));
 		return nullptr;
 	}
 
@@ -47,7 +45,7 @@ FSocket* FSocketSubsystemBSD::CreateSocket(const FName& SocketType, const FStrin
 	bool bIsUDP = SocketType.GetComparisonIndex() == NAME_DGram;
 	int32 SocketTypeFlag = (bIsUDP) ? SOCK_DGRAM : SOCK_STREAM;
 
-	Socket = socket(GetProtocolFamilyValue(ProtocolType), SocketTypeFlag | PlatformSpecificTypeFlags, ((bIsUDP) ? IPPROTO_UDP : IPPROTO_TCP));
+	Socket = socket(GetProtocolFamilyValue(SocketProtocolType), SocketTypeFlag | PlatformSpecificTypeFlags, ((bIsUDP) ? IPPROTO_UDP : IPPROTO_TCP));
 
 #if PLATFORM_ANDROID
 	// To avoid out of range in FD_SET
@@ -58,7 +56,7 @@ FSocket* FSocketSubsystemBSD::CreateSocket(const FName& SocketType, const FStrin
 	else
 #endif
 	{
-		NewSocket = (Socket != INVALID_SOCKET) ? InternalBSDSocketFactory(Socket, ((bIsUDP) ? SOCKTYPE_Datagram : SOCKTYPE_Streaming), SocketDescription, ProtocolType) : nullptr;
+		NewSocket = (Socket != INVALID_SOCKET) ? InternalBSDSocketFactory(Socket, ((bIsUDP) ? SOCKTYPE_Datagram : SOCKTYPE_Streaming), SocketDescription, SocketProtocolType) : nullptr;
 	}
 
 	if (!NewSocket)
@@ -69,43 +67,13 @@ FSocket* FSocketSubsystemBSD::CreateSocket(const FName& SocketType, const FStrin
 	return NewSocket;
 }
 
-ESocketErrors FSocketSubsystemBSD::CreateAddressFromIP(const ANSICHAR* IPAddress, FInternetAddr& OutAddr)
-{
-	FAddressInfoResult GAIResult = GetAddressInfo(ANSI_TO_TCHAR(IPAddress), nullptr,
-		EAddressInfoFlags::AllResultsWithMapping | EAddressInfoFlags::NoResolveHost | EAddressInfoFlags::OnlyUsableAddresses);
-
-	if (GAIResult.Results.Num() > 0)
-	{
-		OutAddr.SetRawIp(GAIResult.Results[0].Address->GetRawIp());
-		return SE_NO_ERROR;
-	}
-	
-	return SE_HOST_NOT_FOUND;
-}
-
-ESocketErrors FSocketSubsystemBSD::GetHostByName(const ANSICHAR* HostName, FInternetAddr& OutAddr)
-{
-	FAddressInfoResult GAIResult = GetAddressInfo(ANSI_TO_TCHAR(HostName), nullptr,
-		EAddressInfoFlags::AllResultsWithMapping | EAddressInfoFlags::OnlyUsableAddresses | EAddressInfoFlags::BindableAddress);
-
-	if (GAIResult.Results.Num() > 0)
-	{
-		TSharedRef<FInternetAddrBSD> ResultAddr = StaticCastSharedRef<FInternetAddrBSD>(GAIResult.Results[0].Address);
-		OutAddr.SetRawIp(ResultAddr->GetRawIp());
-		static_cast<FInternetAddrBSD&>(OutAddr).SetScopeId(ResultAddr->GetScopeId());
-		return SE_NO_ERROR;
-	}
-
-	return SE_HOST_NOT_FOUND;
-}
-
 void FSocketSubsystemBSD::DestroySocket(FSocket* Socket)
 {
 	delete Socket;
 }
 
 FAddressInfoResult FSocketSubsystemBSD::GetAddressInfo(const TCHAR* HostName, const TCHAR* ServiceName,
-	EAddressInfoFlags QueryFlags, ESocketProtocolFamily ProtocolType, ESocketType SocketType)
+	EAddressInfoFlags QueryFlags, const FName ProtocolTypeName, ESocketType SocketType)
 {
 	FAddressInfoResult AddrQueryResult = FAddressInfoResult(HostName, ServiceName);
 
@@ -120,11 +88,17 @@ FAddressInfoResult FSocketSubsystemBSD::GetAddressInfo(const TCHAR* HostName, co
 
 	// Make sure we filter out IPv6 if the platform is not officially supported
 	// (if it isn't supported but we explicitly ask for it, allow it).
-	bool bCanUseIPv6 = (PLATFORM_HAS_BSD_IPV6_SOCKETS || ProtocolType == ESocketProtocolFamily::IPv6) ? true : false;
+	bool bCanUseIPv6 = (PLATFORM_HAS_BSD_IPV6_SOCKETS || ProtocolTypeName == FNetworkProtocolTypes::IPv6) ? true : false;
+
+	// Determine if we can save time with numericserv
+	if (ServiceName != nullptr && FString(ServiceName).IsNumeric())
+	{
+		QueryFlags |= EAddressInfoFlags::NoResolveService;
+	}
 
 	addrinfo HintAddrInfo;
 	FMemory::Memzero(&HintAddrInfo, sizeof(HintAddrInfo));
-	HintAddrInfo.ai_family = GetProtocolFamilyValue(ProtocolType);
+	HintAddrInfo.ai_family = GetProtocolFamilyValue(ProtocolTypeName);
 	HintAddrInfo.ai_flags = GetAddressInfoHintFlag(QueryFlags);
 
 	if (SocketType != ESocketType::SOCKTYPE_Unknown)
@@ -135,10 +109,10 @@ FAddressInfoResult FSocketSubsystemBSD::GetAddressInfo(const TCHAR* HostName, co
 	}
 
 	int32 ErrorCode = getaddrinfo(TCHAR_TO_UTF8(HostName), TCHAR_TO_UTF8(ServiceName), &HintAddrInfo, &AddrInfo);
-	ESocketErrors SocketError = TranslateGAIErrorCode(ErrorCode);
+	AddrQueryResult.ReturnCode = TranslateGAIErrorCode(ErrorCode);
 
 	UE_LOG(LogSockets, Verbose, TEXT("Executed getaddrinfo with HostName: %s Return: %d"), HostName, ErrorCode);
-	if (SocketError == SE_NO_ERROR)
+	if (AddrQueryResult.ReturnCode == SE_NO_ERROR)
 	{
 		addrinfo* AddrInfoHead = AddrInfo;
 		// The canonical name will always be stored in only the first result in a getaddrinfo query
@@ -156,10 +130,13 @@ FAddressInfoResult FSocketSubsystemBSD::GetAddressInfo(const TCHAR* HostName, co
 				{
 					TSharedRef<FInternetAddrBSD> NewAddress = MakeShareable(new FInternetAddrBSD(this));
 					NewAddress->Set(*AddrData, AddrInfo->ai_addrlen);
-					AddrQueryResult.Results.Add(FAddressInfoResultData(NewAddress, AddrInfo->ai_addrlen,
-						GetProtocolFamilyType(AddrInfo->ai_family), GetSocketType(AddrInfo->ai_protocol)));
-
-					UE_LOG(LogSockets, Verbose, TEXT("# Family: %s Address: %s"), ((AddrInfo->ai_family == AF_INET) ? TEXT("IPv4") : TEXT("IPv6")), *(NewAddress->ToString(false)));
+					
+					FAddressInfoResultData NewAddressData(NewAddress, AddrInfo->ai_addrlen, GetProtocolFamilyTypeName(AddrInfo->ai_family), GetSocketType(AddrInfo->ai_protocol));
+					if (AddrQueryResult.Results.Add(NewAddressData) != INDEX_NONE)
+					{
+						sockaddr_in6* IPv6AddrData = reinterpret_cast<sockaddr_in6*>(AddrInfo->ai_addr);
+						UE_LOG(LogSockets, Verbose, TEXT("# Family: %s Address: %s And scope %d"), ((AddrInfo->ai_family == AF_INET) ? TEXT("IPv4") : TEXT("IPv6")), *(NewAddress->ToString(false)), ((AddrInfo->ai_family == AF_INET) ? -1 : IPv6AddrData->sin6_scope_id));
+					}
 				}
 			}
 		}
@@ -167,7 +144,7 @@ FAddressInfoResult FSocketSubsystemBSD::GetAddressInfo(const TCHAR* HostName, co
 	}
 	else
 	{
-		UE_LOG(LogSockets, Warning, TEXT("GetAddressInfo failed to resolve host with error %s [%d]"), GetSocketError(SocketError), ErrorCode);
+		UE_LOG(LogSockets, Warning, TEXT("GetAddressInfo failed to resolve host with error %s [%d]"), GetSocketError(AddrQueryResult.ReturnCode), ErrorCode);
 	}
 #else
 	UE_LOG(LogSockets, Error, TEXT("Platform has no getaddrinfo(), but did not override FSocketSubsystem::GetAddressInfo()"));
@@ -175,6 +152,40 @@ FAddressInfoResult FSocketSubsystemBSD::GetAddressInfo(const TCHAR* HostName, co
 	return AddrQueryResult;
 }
 
+TSharedPtr<FInternetAddr> FSocketSubsystemBSD::GetAddressFromString(const FString& InAddress)
+{
+	sockaddr_storage NetworkBuffer;
+	FMemory::Memzero(NetworkBuffer);
+
+	void* UniversalNetBufferPtr = nullptr;
+	int32 AddrFamily = AF_INET;
+
+#if PLATFORM_HAS_BSD_IPV6_SOCKETS
+	int32 ColonIndex;
+	if (InAddress.FindChar(TEXT(':'), ColonIndex))
+	{
+		AddrFamily = AF_INET6;
+		UniversalNetBufferPtr = &(((sockaddr_in6*)&NetworkBuffer)->sin6_addr);
+	}
+	else
+#endif
+	{
+		UniversalNetBufferPtr = &(((sockaddr_in*)&NetworkBuffer)->sin_addr);
+	}
+
+	NetworkBuffer.ss_family = AddrFamily;
+
+	if (inet_pton(AddrFamily, TCHAR_TO_ANSI(*InAddress), UniversalNetBufferPtr))
+	{
+		TSharedRef<FInternetAddrBSD> ReturnAddress = StaticCastSharedRef<FInternetAddrBSD>(CreateInternetAddr());
+		ReturnAddress->Set(NetworkBuffer);
+		return ReturnAddress;
+	}
+	
+	UE_LOG(LogSockets, Warning, TEXT("Could not serialize %s, got error code %d"), *InAddress, GetLastErrorCode());
+	return nullptr;
+}
+
 bool FSocketSubsystemBSD::GetHostName(FString& HostName)
 {
 #if PLATFORM_HAS_BSD_SOCKET_FEATURE_GETHOSTNAME
@@ -191,18 +202,40 @@ bool FSocketSubsystemBSD::GetHostName(FString& HostName)
 #endif
 }
 
+TSharedRef<FInternetAddr> FSocketSubsystemBSD::GetLocalBindAddr(FOutputDevice& Out)
+{
+	bool bCanBindAll;
+	// look up the local host address
+	TSharedRef<FInternetAddr> BindAddr = GetLocalHostAddr(Out, bCanBindAll);
+	if (bCanBindAll)
+	{
+		// Pass in the result from looking up the machine name as the protocol 
+		// to use for the any address.
+		// The idea behind this is if your platform has a default of IPv4 addresses,
+		// but the runtime machine has a IPv6 address, then the wrong address protocol
+		// will be used when using the any address.
+		TSharedRef<FInternetAddrBSD> BSDAddr = StaticCastSharedRef<FInternetAddrBSD>(BindAddr);
+		if (BindAddr->GetProtocolType() == FNetworkProtocolTypes::IPv6)
+		{
+			BSDAddr->SetAnyIPv6Address();
+		}
+		else
+		{
+			BSDAddr->SetAnyIPv4Address();
+		}
+	}
+
+	return BindAddr;
+}
 
 const TCHAR* FSocketSubsystemBSD::GetSocketAPIName() const
 {
 	return TEXT("BSD IPv4/6");
 }
 
-TSharedRef<FInternetAddr> FSocketSubsystemBSD::CreateInternetAddr(uint32 Address, uint32 Port)
+TSharedRef<FInternetAddr> FSocketSubsystemBSD::CreateInternetAddr()
 {
-	TSharedRef<FInternetAddr> ReturnAddr = MakeShareable(new FInternetAddrBSD(this));
-	ReturnAddr->SetIp(Address);
-	ReturnAddr->SetPort(Port);
-	return ReturnAddr;
+	return MakeShareable(new FInternetAddrBSD(this));
 }
 
 bool FSocketSubsystemBSD::IsSocketWaitSupported() const
@@ -215,7 +248,6 @@ ESocketErrors FSocketSubsystemBSD::GetLastErrorCode()
 	return TranslateErrorCode(errno);
 }
 
-
 ESocketErrors FSocketSubsystemBSD::TranslateErrorCode(int32 Code)
 {
 	// @todo sockets: Windows for some reason doesn't seem to have all of the standard error messages,
@@ -329,6 +361,32 @@ ESocketErrors FSocketSubsystemBSD::TranslateGAIErrorCode(int32 Code) const
 	return SE_NO_ERROR;
 }
 
+int32 FSocketSubsystemBSD::GetProtocolFamilyValue(const FName& InProtocol) const
+{
+	if (InProtocol == FNetworkProtocolTypes::IPv4)
+	{
+		return AF_INET;
+	}
+	else if (InProtocol == FNetworkProtocolTypes::IPv6)
+	{
+		return AF_INET6;
+	}
+
+	return AF_UNSPEC;
+}
+
+const FName FSocketSubsystemBSD::GetProtocolFamilyTypeName(int32 InProtocol) const
+{
+	switch (InProtocol)
+	{
+		default:
+		case AF_UNSPEC: return NAME_None;
+		case AF_INET:   return FNetworkProtocolTypes::IPv4;
+		case AF_INET6:  return FNetworkProtocolTypes::IPv6;
+	}
+}
+
+PRAGMA_DISABLE_DEPRECATION_WARNINGS
 int32 FSocketSubsystemBSD::GetProtocolFamilyValue(ESocketProtocolFamily InProtocol) const
 {
 	switch (InProtocol)
@@ -350,6 +408,7 @@ ESocketProtocolFamily FSocketSubsystemBSD::GetProtocolFamilyType(int32 InProtoco
 		case AF_INET6:  return ESocketProtocolFamily::IPv6;
 	}
 }
+PRAGMA_ENABLE_DEPRECATION_WARNINGS
 
 ESocketType FSocketSubsystemBSD::GetSocketType(int32 InSocketType) const
 {
diff --git a/Engine/Source/Runtime/Sockets/Private/BSDSockets/SocketSubsystemBSD.h b/Engine/Source/Runtime/Sockets/Private/BSDSockets/SocketSubsystemBSD.h
--- a/Engine/Source/Runtime/Sockets/Private/BSDSockets/SocketSubsystemBSD.h
+++ b/Engine/Source/Runtime/Sockets/Private/BSDSockets/SocketSubsystemBSD.h
@@ -6,10 +6,10 @@
 #include "Misc/OutputDeviceRedirector.h"
 #include "SocketSubsystemPackage.h"
 #include "SocketSubsystem.h"
+#include "IPAddress.h"
 
 #if PLATFORM_HAS_BSD_SOCKETS || PLATFORM_HAS_BSD_IPV6_SOCKETS
 
-class FInternetAddr;
 class FSocketBSD;
 
 #include "SocketSubsystemBSDPrivate.h"
@@ -20,9 +20,6 @@ class FSocketBSD;
 class FSocketSubsystemBSD : public ISocketSubsystem
 {
 public:
-
-	virtual ESocketErrors CreateAddressFromIP(const ANSICHAR* IPAddress, FInternetAddr& OutAddr);
-
 	/**
 	 * Specifies the default socket protocol family to use when creating a socket
 	 * without explicitly passing in the protocol type on creation.
@@ -33,24 +30,32 @@ public:
 	 * All sockets created using the base class's CreateSocket will use this function
 	 * to determine domain.
 	 */
-	virtual ESocketProtocolFamily GetDefaultSocketProtocolFamily() const
+	virtual FName GetDefaultSocketProtocolFamily() const
 	{
-		return ESocketProtocolFamily::IPv4;
+		return FNetworkProtocolTypes::IPv4;
 	}
 
+
 	// ISocketSubsystem interface
-	virtual TSharedRef<FInternetAddr> CreateInternetAddr(uint32 Address = 0, uint32 Port = 0) override;
-	virtual class FSocket* CreateSocket( const FName& SocketType, const FString& SocketDescription, bool bForceUDP = false ) override;
-	virtual class FSocket* CreateSocket(const FName& SocketType, const FString& SocketDescription, ESocketProtocolFamily ProtocolType) override;
+	virtual TSharedRef<FInternetAddr> CreateInternetAddr() override;
+
+	virtual class FSocket* CreateSocket(const FName& SocketType, const FString& SocketDescription, bool bForceUDP = false) override
+	{
+		return CreateSocket(SocketType, SocketDescription, GetDefaultSocketProtocolFamily());
+	}
+
+	virtual class FSocket* CreateSocket(const FName& SocketType, const FString& SocketDescription, const FName& ProtocolType) override;
+
 	virtual void DestroySocket( class FSocket* Socket ) override;
 
 	virtual FAddressInfoResult GetAddressInfo(const TCHAR* HostName, const TCHAR* ServiceName = nullptr,
 		EAddressInfoFlags QueryFlags = EAddressInfoFlags::Default,
-		ESocketProtocolFamily ProtocolType = ESocketProtocolFamily::None,
+		const FName ProtocolTypeName = NAME_None,
 		ESocketType SocketType = ESocketType::SOCKTYPE_Unknown) override;
 
-	virtual ESocketErrors GetHostByName( const ANSICHAR* HostName, FInternetAddr& OutAddr ) override;
-	virtual bool GetHostName( FString& HostName ) override;
+	virtual TSharedPtr<FInternetAddr> GetAddressFromString(const FString& InAddress) override;
+
+	virtual bool GetHostName(FString& HostName) override;
 	virtual ESocketErrors GetLastErrorCode() override;
 
 	virtual bool GetLocalAdapterAddresses( TArray<TSharedPtr<FInternetAddr> >& OutAdresses ) override
@@ -62,6 +67,8 @@ public:
 		return true;
 	}
 
+	virtual TSharedRef<FInternetAddr> GetLocalBindAddr(FOutputDevice& Out) override;
+
 	virtual const TCHAR* GetSocketAPIName() const override;
 
 	virtual bool RequiresChatDataBeSeparate() override
@@ -79,14 +86,20 @@ public:
 	/**
 	 * Translates an ESocketProtocolFamily code into a value usable by raw socket apis.
 	 */
+	UE_DEPRECATED(4.23, "Switch to the FName version for scalable protocol support")
 	virtual int32 GetProtocolFamilyValue(ESocketProtocolFamily InProtocol) const;
+	virtual int32 GetProtocolFamilyValue(const FName& InProtocol) const;
 	
 	virtual bool IsSocketWaitSupported() const override;
 	
 	/**
-	 * Translates an raw socket family type value into an enum that can be used by the network layer.
+	 * Translates an raw socket family type value into a protocol name that can be used by the network layer.
 	 */
+	virtual const FName GetProtocolFamilyTypeName(int32 InProtocol) const;
+
+	UE_DEPRECATED(4.23, "Use GetProtocolFamilyTypeName")
 	virtual ESocketProtocolFamily GetProtocolFamilyType(int32 InProtocol) const;
+	
 
 	/**
 	 * Translates an raw socket protocol type value into an enum that can be used by the network layer.
@@ -107,7 +120,7 @@ protected:
 	/**
 	 * Allows a subsystem subclass to create a FSocketBSD sub class.
 	 */
-	virtual class FSocketBSD* InternalBSDSocketFactory( SOCKET Socket, ESocketType SocketType, const FString& SocketDescription, ESocketProtocolFamily SocketProtocol);
+	virtual class FSocketBSD* InternalBSDSocketFactory( SOCKET Socket, ESocketType SocketType, const FString& SocketDescription, const FName& SocketProtocol);
 
 	/**
 	 * Allows a subsystem subclass to create a FSocketBSD sub class.
diff --git a/Engine/Source/Runtime/Sockets/Private/BSDSockets/SocketsBSD.cpp b/Engine/Source/Runtime/Sockets/Private/BSDSockets/SocketsBSD.cpp
--- a/Engine/Source/Runtime/Sockets/Private/BSDSockets/SocketsBSD.cpp
+++ b/Engine/Source/Runtime/Sockets/Private/BSDSockets/SocketsBSD.cpp
@@ -66,6 +66,13 @@ bool FSocketBSD::Close(void)
 
 bool FSocketBSD::Bind(const FInternetAddr& Addr)
 {
+	if (Addr.GetProtocolType() != GetProtocol())
+	{
+		UE_LOG(LogSockets, Warning, TEXT("Tried to bind address with protocol %s to a socket with protocol %s"),
+			*Addr.GetProtocolType().ToString(), *GetProtocol().ToString());
+		return false;
+	}
+
 	const FInternetAddrBSD& BSDAddr = static_cast<const FInternetAddrBSD&>(Addr);
 	return bind(Socket, (const sockaddr*)&(BSDAddr.Addr), BSDAddr.GetStorageSize()) == 0;
 }
@@ -73,6 +80,13 @@ bool FSocketBSD::Bind(const FInternetAddr& Addr)
 
 bool FSocketBSD::Connect(const FInternetAddr& Addr)
 {
+	if (Addr.GetProtocolType() != GetProtocol())
+	{
+		UE_LOG(LogSockets, Warning, TEXT("Tried to connect with an address with protocol %s using a socket with protocol %s"), 
+			*Addr.GetProtocolType().ToString(), *GetProtocol().ToString());
+		return false;
+	}
+
 	const FInternetAddrBSD& BSDAddr = static_cast<const FInternetAddrBSD&>(Addr);
 	int32 Return = connect(Socket, (const sockaddr*)&(BSDAddr.Addr), BSDAddr.GetStorageSize());
 	
@@ -157,7 +171,7 @@ FSocket* FSocketBSD::Accept(FInternetAddr& OutAddr, const FString& InSocketDescr
 		// we need the subclass to create the actual FSocket object
 		check(SocketSubsystem);
 		FSocketSubsystemBSD* BSDSystem = static_cast<FSocketSubsystemBSD*>(SocketSubsystem);
-		return BSDSystem->InternalBSDSocketFactory(NewSocket, SocketType, InSocketDescription, BSDAddr.GetProtocolFamily());
+		return BSDSystem->InternalBSDSocketFactory(NewSocket, SocketType, InSocketDescription, BSDAddr.GetProtocolType());
 	}
 
 	return NULL;
@@ -166,6 +180,12 @@ FSocket* FSocketBSD::Accept(FInternetAddr& OutAddr, const FString& InSocketDescr
 
 bool FSocketBSD::SendTo(const uint8* Data, int32 Count, int32& BytesSent, const FInternetAddr& Destination)
 {
+	// TODO: Consider converting IPv4 addresses to v6 when needed
+	if (Destination.GetProtocolType() != GetProtocol())
+	{
+		return false;
+	}
+
 	const FInternetAddrBSD& BSDAddr = static_cast<const FInternetAddrBSD&>(Destination);
 	// Write the data and see how much was written
 	BytesSent = sendto(Socket, (const char*)Data, Count, 0, (const sockaddr*)&(BSDAddr.Addr), BSDAddr.GetStorageSize());
@@ -383,7 +403,7 @@ bool FSocketBSD::JoinMulticastGroup(const FInternetAddr& GroupAddress)
 	const FInternetAddrBSD& BSDAddr = static_cast<const FInternetAddrBSD&>(GroupAddress);
 
 #if PLATFORM_HAS_BSD_IPV6_SOCKETS
-	if (BSDAddr.GetProtocolFamily() == ESocketProtocolFamily::IPv6)
+	if (BSDAddr.GetProtocolType() == FNetworkProtocolTypes::IPv6)
 	{
 		ipv6_mreq imr;
 		imr.ipv6mr_interface = htonl(BSDAddr.GetScopeId());
@@ -405,7 +425,7 @@ bool FSocketBSD::JoinMulticastGroup(const FInternetAddr& GroupAddress, const FIn
 	const FInternetAddrBSD& BSDIFAddr = static_cast<const FInternetAddrBSD&>(InterfaceAddress);
 
 #if PLATFORM_HAS_BSD_IPV6_SOCKETS
-	if (BSDAddr.GetProtocolFamily() == ESocketProtocolFamily::IPv6)
+	if (BSDAddr.GetProtocolType() == FNetworkProtocolTypes::IPv6)
 	{
 		ipv6_mreq imr;
 		imr.ipv6mr_interface = htonl(BSDIFAddr.GetScopeId());
@@ -426,7 +446,7 @@ bool FSocketBSD::LeaveMulticastGroup(const FInternetAddr& GroupAddress)
 	const FInternetAddrBSD& BSDAddr = static_cast<const FInternetAddrBSD&>(GroupAddress);
 
 #if PLATFORM_HAS_BSD_IPV6_SOCKETS
-	if (BSDAddr.GetProtocolFamily() == ESocketProtocolFamily::IPv6)
+	if (BSDAddr.GetProtocolType() == FNetworkProtocolTypes::IPv6)
 	{
 		ipv6_mreq imr;
 		imr.ipv6mr_interface = htonl(BSDAddr.GetScopeId());
@@ -448,7 +468,7 @@ bool FSocketBSD::LeaveMulticastGroup(const FInternetAddr& GroupAddress, const FI
 	const FInternetAddrBSD& BSDIFAddr = static_cast<const FInternetAddrBSD&>(InterfaceAddress);
 
 #if PLATFORM_HAS_BSD_IPV6_SOCKETS
-	if (BSDAddr.GetProtocolFamily() == ESocketProtocolFamily::IPv6)
+	if (BSDAddr.GetProtocolType() == FNetworkProtocolTypes::IPv6)
 	{
 		ipv6_mreq imr;
 		imr.ipv6mr_interface = htonl(BSDIFAddr.GetScopeId());
@@ -467,7 +487,7 @@ bool FSocketBSD::LeaveMulticastGroup(const FInternetAddr& GroupAddress, const FI
 bool FSocketBSD::SetMulticastLoopback(bool bLoopback)
 {
 #if PLATFORM_HAS_BSD_IPV6_SOCKETS
-	if (SocketProtocol == ESocketProtocolFamily::IPv6)
+	if (SocketProtocol == FNetworkProtocolTypes::IPv6)
 	{
 		uint32 ShouldLoopback = bLoopback ? 1 : 0;
 		return (setsockopt(Socket, IPPROTO_IPV6, IP_MULTICAST_LOOP, (char*)&ShouldLoopback, sizeof(ShouldLoopback)) == 0);
@@ -481,7 +501,7 @@ bool FSocketBSD::SetMulticastLoopback(bool bLoopback)
 bool FSocketBSD::SetMulticastTtl(uint8 TimeToLive)
 {
 #if PLATFORM_HAS_BSD_IPV6_SOCKETS
-	if (SocketProtocol == ESocketProtocolFamily::IPv6)
+	if (SocketProtocol == FNetworkProtocolTypes::IPv6)
 	{
 		uint32 RealTimeToLive = TimeToLive;
 		return (setsockopt(Socket, IPPROTO_IPV6, IPV6_MULTICAST_HOPS, (char*)&RealTimeToLive, sizeof(RealTimeToLive)) == 0);
@@ -497,7 +517,7 @@ bool FSocketBSD::SetMulticastInterface(const FInternetAddr& InterfaceAddress)
 	const FInternetAddrBSD& BSDIFAddr = static_cast<const FInternetAddrBSD&>(InterfaceAddress);
 
 #if PLATFORM_HAS_BSD_IPV6_SOCKETS
-	if (BSDIFAddr.GetProtocolFamily() == ESocketProtocolFamily::IPv6)
+	if (BSDIFAddr.GetProtocolType() == FNetworkProtocolTypes::IPv6)
 	{
 		uint32 InterfaceIndex = htonl(BSDIFAddr.GetScopeId());
 		return (setsockopt(Socket, IPPROTO_IPV6, IPV6_MULTICAST_IF, (char*)&InterfaceIndex, sizeof(InterfaceIndex)) == 0);
@@ -595,16 +615,23 @@ int32 FSocketBSD::GetPortNo(void)
 bool FSocketBSD::SetIPv6Only(bool bIPv6Only)
 {
 #if PLATFORM_HAS_BSD_IPV6_SOCKETS
-	int v6only = bIPv6Only ? 1 : 0;
-	bool bOk = (setsockopt(Socket, IPPROTO_IPV6, IPV6_V6ONLY, (char*)&v6only, sizeof(v6only)) == 0);
+	if (GetProtocol() == FNetworkProtocolTypes::IPv6)
+	{
+		int v6only = bIPv6Only ? 1 : 0;
+		bool bOk = (setsockopt(Socket, IPPROTO_IPV6, IPV6_V6ONLY, (char*)&v6only, sizeof(v6only)) == 0);
 
-	if (bOk == false)
+		if (bOk == false)
+		{
+			check(SocketSubsystem);
+			UE_LOG(LogSockets, Error, TEXT("Failed to set sock opt for socket (%s)"), SocketSubsystem->GetSocketError());
+		}
+
+		return bOk;
+	}
+	else
 	{
-		check(SocketSubsystem);
-		UE_LOG(LogSockets, Error, TEXT("Failed to set sock opt for socket (%s)"), SocketSubsystem->GetSocketError());
+		return true;
 	}
-
-	return bOk;
 #else 
 	return false;
 #endif
diff --git a/Engine/Source/Runtime/Sockets/Private/BSDSockets/SocketsBSD.h b/Engine/Source/Runtime/Sockets/Private/BSDSockets/SocketsBSD.h
--- a/Engine/Source/Runtime/Sockets/Private/BSDSockets/SocketsBSD.h
+++ b/Engine/Source/Runtime/Sockets/Private/BSDSockets/SocketsBSD.h
@@ -52,7 +52,7 @@ public:
 	 */
 	UE_DEPRECATED(4.22, "Use the socket constructor that specifies protocol stack for better compatibility and debugging")
 	FSocketBSD(SOCKET InSocket, ESocketType InSocketType, const FString& InSocketDescription, ISocketSubsystem * InSubsystem)
-		: FSocket(InSocketType, InSocketDescription, ESocketProtocolFamily::None)
+		: FSocket(InSocketType, InSocketDescription, NAME_None)
 		, Socket(InSocket)
 		, LastActivityTime(0)
 		, SocketSubsystem(InSubsystem)
@@ -66,7 +66,7 @@ public:
 	 * @param InSocketDescription the debug description of the socket.
 	 * @param InSocketProtocol the protocol this socket is initialized with
 	 */
-	FSocketBSD(SOCKET InSocket, ESocketType InSocketType, const FString& InSocketDescription, ESocketProtocolFamily InSocketProtocol, ISocketSubsystem * InSubsystem)
+	FSocketBSD(SOCKET InSocket, ESocketType InSocketType, const FString& InSocketDescription, const FName& InSocketProtocol, ISocketSubsystem * InSubsystem)
 		: FSocket(InSocketType, InSocketDescription, InSocketProtocol)
 		, Socket(InSocket)
 		, LastActivityTime(0.0)
diff --git a/Engine/Source/Runtime/Sockets/Private/IOS/IPAddressBSDIOS.h b/Engine/Source/Runtime/Sockets/Private/IOS/IPAddressBSDIOS.h
--- a/Engine/Source/Runtime/Sockets/Private/IOS/IPAddressBSDIOS.h
+++ b/Engine/Source/Runtime/Sockets/Private/IOS/IPAddressBSDIOS.h
@@ -21,13 +21,17 @@ public:
 		FSocketSubsystemBSD* SocketSubsystem = static_cast<FSocketSubsystemBSD*>(ISocketSubsystem::Get(PLATFORM_SOCKETSUBSYSTEM));
 		if (SocketSubsystem)
 		{
-			// Do a query to get the scope id of the address.
-			if (SocketSubsystem->CreateAddressFromIP("ff02::01", *this) != SE_NO_ERROR)
+			TSharedPtr<FInternetAddrBSD> MultiCastAddr = SocketSubsystem->GetAddressFromString(TEXT("ff02::01"));
+			if (!MultiCastAddr.IsValid())
 			{
 				UE_LOG(LogSockets, Warning, TEXT("Could not resolve the broadcast address for iOS, this address will just be blank"));
 			}
 			else
 			{
+				// Set the address from the query
+				SetRawIp(MultiCastAddr->GetRawIp());
+
+				// Do a query to get the scope id of the address.
 				bool bUnusedBool;
 				TSharedRef<FInternetAddrBSDIOS> ScopeAddr = 
 					StaticCastSharedRef<FInternetAddrBSDIOS>(SocketSubsystem->GetLocalHostAddr(*GLog, bUnusedBool));
diff --git a/Engine/Source/Runtime/Sockets/Private/IOS/SocketSubsystemIOS.cpp b/Engine/Source/Runtime/Sockets/Private/IOS/SocketSubsystemIOS.cpp
--- a/Engine/Source/Runtime/Sockets/Private/IOS/SocketSubsystemIOS.cpp
+++ b/Engine/Source/Runtime/Sockets/Private/IOS/SocketSubsystemIOS.cpp
@@ -3,16 +3,14 @@
 #include "SocketSubsystemIOS.h"
 #include "SocketSubsystemModule.h"
 #include "Modules/ModuleManager.h"
-#include "BSDSockets/SocketsBSD.h"
-#include "IPAddress.h"
-#include <net/if.h>
-#include <ifaddrs.h>
 #include "SocketsBSDIOS.h"
 #include "IPAddressBSDIOS.h"
+#include <net/if.h>
+#include <ifaddrs.h>
 
 FSocketSubsystemIOS* FSocketSubsystemIOS::SocketSingleton = NULL;
 
-class FSocketBSD* FSocketSubsystemIOS::InternalBSDSocketFactory(SOCKET Socket, ESocketType SocketType, const FString& SocketDescription, ESocketProtocolFamily SocketProtocol)
+class FSocketBSD* FSocketSubsystemIOS::InternalBSDSocketFactory(SOCKET Socket, ESocketType SocketType, const FString& SocketDescription, const FName& SocketProtocol)
 {
 	UE_LOG(LogIOS, Log, TEXT(" FSocketSubsystemIOS::InternalBSDSocketFactory"));
 	return new FSocketBSDIOS(Socket, SocketType, SocketDescription, SocketProtocol, this);
@@ -77,33 +75,12 @@ bool FSocketSubsystemIOS::HasNetworkDevice()
 	return true;
 }
 
-ESocketErrors FSocketSubsystemIOS::CreateAddressFromIP(const ANSICHAR* IPAddress, FInternetAddr& OutAddr)
-{
-	return GetHostByName(IPAddress, OutAddr);
-}
-
-ESocketErrors FSocketSubsystemIOS::GetHostByName(const ANSICHAR* HostName, FInternetAddr& OutAddr)
-{
-	FAddressInfoResult GAIResult = GetAddressInfo(ANSI_TO_TCHAR(HostName), nullptr,
-		EAddressInfoFlags::AllResultsWithMapping | EAddressInfoFlags::OnlyUsableAddresses, ESocketProtocolFamily::IPv6);
-
-	if (GAIResult.Results.Num() > 0)
-	{
-		TSharedRef<FInternetAddrBSDIOS> ResultAddr = StaticCastSharedRef<FInternetAddrBSDIOS>(GAIResult.Results[0].Address);
-		OutAddr.SetRawIp(ResultAddr->GetRawIp());
-		static_cast<FInternetAddrBSDIOS&>(OutAddr).SetScopeId(ResultAddr->GetScopeId());
-		return SE_NO_ERROR;
-	}
-
-	return SE_HOST_NOT_FOUND;
-}
-
-FSocket* FSocketSubsystemIOS::CreateSocket(const FName& SocketType, const FString& SocketDescription, ESocketProtocolFamily ProtocolType)
+FSocket* FSocketSubsystemIOS::CreateSocket(const FName& SocketType, const FString& SocketDescription, const FName& ProtocolType)
 {
 	FSocketBSD* NewSocket = (FSocketBSD*)FSocketSubsystemBSD::CreateSocket(SocketType, SocketDescription, ProtocolType);
 	if (NewSocket)
 	{
-		if (ProtocolType != ESocketProtocolFamily::IPv4)
+		if (ProtocolType != FNetworkProtocolTypes::IPv4)
 		{
 			NewSocket->SetIPv6Only(false);
 		}
@@ -174,7 +151,7 @@ TSharedRef<FInternetAddr> FSocketSubsystemIOS::GetLocalHostAddr(FOutputDevice& O
 				}
 			}
 		}
-		
+
 		// Free memory
 		freeifaddrs(Interfaces);
 
@@ -197,10 +174,7 @@ TSharedRef<FInternetAddr> FSocketSubsystemIOS::GetLocalHostAddr(FOutputDevice& O
 	return HostAddr;
 }
 
-TSharedRef<FInternetAddr> FSocketSubsystemIOS::CreateInternetAddr(uint32 Address, uint32 Port)
+TSharedRef<FInternetAddr> FSocketSubsystemIOS::CreateInternetAddr()
 {
-	TSharedRef<FInternetAddr> Result = MakeShareable(new FInternetAddrBSDIOS(this));
-	Result->SetIp(Address);
-	Result->SetPort(Port);
-	return Result;
+	return MakeShareable(new FInternetAddrBSDIOS(this));
 }
diff --git a/Engine/Source/Runtime/Sockets/Private/IOS/SocketSubsystemIOS.h b/Engine/Source/Runtime/Sockets/Private/IOS/SocketSubsystemIOS.h
--- a/Engine/Source/Runtime/Sockets/Private/IOS/SocketSubsystemIOS.h
+++ b/Engine/Source/Runtime/Sockets/Private/IOS/SocketSubsystemIOS.h
@@ -12,7 +12,7 @@
 class FSocketSubsystemIOS : public FSocketSubsystemBSD
 {
 protected:
-	virtual TSharedRef<FInternetAddr> CreateInternetAddr(uint32 Address = 0, uint32 Port = 0) override;
+	virtual TSharedRef<FInternetAddr> CreateInternetAddr() override;
 	/** Single instantiation of this subsystem */
 	static FSocketSubsystemIOS* SocketSingleton;
 
@@ -35,9 +35,9 @@ public:
 	static void Destroy();
 
 
-	virtual ESocketProtocolFamily GetDefaultSocketProtocolFamily() const override
+	virtual FName GetDefaultSocketProtocolFamily() const override
 	{
-		return ESocketProtocolFamily::IPv6;
+		return FNetworkProtocolTypes::IPv6;
 	}
 
 public:
@@ -51,14 +51,11 @@ public:
 	{
 	}
 
-	virtual ESocketErrors CreateAddressFromIP(const ANSICHAR* IPAddress, FInternetAddr& OutAddr) override;
-	virtual ESocketErrors GetHostByName(const ANSICHAR* HostName, FInternetAddr& OutAddr) override;
-
 	virtual bool Init(FString& Error) override;
 	virtual void Shutdown() override;
 	virtual bool HasNetworkDevice() override;
-	virtual FSocket* CreateSocket(const FName& SocketType, const FString& SocketDescription, ESocketProtocolFamily ProtocolType) override;
+	virtual FSocket* CreateSocket(const FName& SocketType, const FString& SocketDescription, const FName& ProtocolType) override;
 
 	virtual TSharedRef<FInternetAddr> GetLocalHostAddr(FOutputDevice& Out, bool& bCanBindAll) override;
-	virtual class FSocketBSD* InternalBSDSocketFactory(SOCKET Socket, ESocketType SocketType, const FString& SocketDescription, ESocketProtocolFamily SocketProtocol) override;
+	virtual class FSocketBSD* InternalBSDSocketFactory(SOCKET Socket, ESocketType SocketType, const FString& SocketDescription, const FName& SocketProtocol) override;
 };
\ No newline at end of file
diff --git a/Engine/Source/Runtime/Sockets/Private/IOS/SocketsBSDIOS.h b/Engine/Source/Runtime/Sockets/Private/IOS/SocketsBSDIOS.h
--- a/Engine/Source/Runtime/Sockets/Private/IOS/SocketsBSDIOS.h
+++ b/Engine/Source/Runtime/Sockets/Private/IOS/SocketsBSDIOS.h
@@ -9,9 +9,6 @@
 
 class FInternetAddr;
 
-#if PLATFORM_HAS_BSD_IPV6_SOCKETS
-
-
 /**
  * Implements a BSD network socket on IOS.
  */
@@ -20,7 +17,7 @@ class FSocketBSDIOS
 {
 public:
 
-	FSocketBSDIOS(SOCKET InSocket, ESocketType InSocketType, const FString& InSocketDescription, ESocketProtocolFamily InSocketProtocol, ISocketSubsystem* InSubsystem)
+	FSocketBSDIOS(SOCKET InSocket, ESocketType InSocketType, const FString& InSocketDescription, const FName& InSocketProtocol, ISocketSubsystem* InSubsystem)
 		:FSocketBSD(InSocket, InSocketType, InSocketDescription, InSocketProtocol, InSubsystem)
 	{
 	}
@@ -31,5 +28,3 @@ public:
 	}
 
 };
-
-#endif
diff --git a/Engine/Source/Runtime/Sockets/Private/IPAddress.cpp b/Engine/Source/Runtime/Sockets/Private/IPAddress.cpp
--- a/Engine/Source/Runtime/Sockets/Private/IPAddress.cpp
+++ b/Engine/Source/Runtime/Sockets/Private/IPAddress.cpp
@@ -2,6 +2,7 @@
 
 #include "IPAddress.h"
 #include "SocketSubsystem.h"
+#include "AddressInfoTypes.h"
 
 /**
  * Sets the address to return to the caller
@@ -34,12 +35,12 @@ void FResolveInfoAsync::DoWork()
 	int32 AttemptCount = 0;
 
 	ISocketSubsystem* SocketSubsystem = ISocketSubsystem::Get();
-	Addr = SocketSubsystem->CreateInternetAddr(0,0);
 
 	// Make up to 3 attempts to resolve it
 	do 
 	{
-		ErrorCode = SocketSubsystem->GetHostByName(HostName, *Addr);
+		FAddressInfoResult GAIResults = SocketSubsystem->GetAddressInfo(ANSI_TO_TCHAR(HostName), nullptr, EAddressInfoFlags::Default, NAME_None);
+		ErrorCode = GAIResults.ReturnCode;
 		if (ErrorCode != SE_NO_ERROR)
 		{
 			if (ErrorCode == SE_HOST_NOT_FOUND || ErrorCode == SE_NO_DATA || ErrorCode == SE_ETIMEDOUT)
@@ -48,12 +49,19 @@ void FResolveInfoAsync::DoWork()
 				AttemptCount = 3;
 			}
 		}
+		else if(GAIResults.Results.Num() > 0)
+		{
+			
+			// Pull out the address
+			Addr = GAIResults.Results[0].Address;
+
+			// Cache for reuse
+			SocketSubsystem->AddHostNameToCache(HostName, Addr);
+
+			return;
+		}
+
 		AttemptCount++;
 	}
 	while (ErrorCode != SE_NO_ERROR && AttemptCount < 3 && bShouldAbandon == false);
-	if (ErrorCode == SE_NO_ERROR)
-	{
-		// Cache for reuse
-		SocketSubsystem->AddHostNameToCache(HostName, Addr);
-	}
 }
diff --git a/Engine/Source/Runtime/Sockets/Private/Mac/SocketSubsystemMac.cpp b/Engine/Source/Runtime/Sockets/Private/Mac/SocketSubsystemMac.cpp
--- a/Engine/Source/Runtime/Sockets/Private/Mac/SocketSubsystemMac.cpp
+++ b/Engine/Source/Runtime/Sockets/Private/Mac/SocketSubsystemMac.cpp
@@ -66,7 +66,7 @@ bool FSocketSubsystemMac::HasNetworkDevice()
 	return true;
 }
 
-class FSocketBSD* FSocketSubsystemMac::InternalBSDSocketFactory(SOCKET Socket, ESocketType SocketType, const FString& SocketDescription, ESocketProtocolFamily SocketProtocol)
+class FSocketBSD* FSocketSubsystemMac::InternalBSDSocketFactory(SOCKET Socket, ESocketType SocketType, const FString& SocketDescription, const FName& SocketProtocol)
 {
 	// return a new socket object
 	FSocketMac* MacSock = new FSocketMac(Socket, SocketType, SocketDescription, SocketProtocol, this);
diff --git a/Engine/Source/Runtime/Sockets/Private/Mac/SocketSubsystemMac.h b/Engine/Source/Runtime/Sockets/Private/Mac/SocketSubsystemMac.h
--- a/Engine/Source/Runtime/Sockets/Private/Mac/SocketSubsystemMac.h
+++ b/Engine/Source/Runtime/Sockets/Private/Mac/SocketSubsystemMac.h
@@ -36,7 +36,7 @@ PACKAGE_SCOPE:
 	/**
 	 * Allows a subsystem subclass to create a FSocketBSD sub class 
 	 */
-	virtual class FSocketBSD* InternalBSDSocketFactory(SOCKET Socket, ESocketType SocketType, const FString& SocketDescription, ESocketProtocolFamily SocketProtocol);
+	virtual class FSocketBSD* InternalBSDSocketFactory(SOCKET Socket, ESocketType SocketType, const FString& SocketDescription, const FName& SocketProtocol);
 
 public:
 
@@ -83,7 +83,7 @@ public:
 	 * @param InSocketType the type of socket that was created
 	 * @param InSocketDescription the debug description of the socket
 	 */
-	FSocketMac(SOCKET InSocket, ESocketType InSocketType, const FString& InSocketDescription, ESocketProtocolFamily InSocketProtocol, ISocketSubsystem * InSubsystem) 
+	FSocketMac(SOCKET InSocket, ESocketType InSocketType, const FString& InSocketDescription, const FName& InSocketProtocol, ISocketSubsystem * InSubsystem) 
 	: FSocketBSD(InSocket, InSocketType, InSocketDescription, InSocketProtocol, InSubsystem)
 	{
 	}
diff --git a/Engine/Source/Runtime/Sockets/Private/SocketSubsystem.cpp b/Engine/Source/Runtime/Sockets/Private/SocketSubsystem.cpp
--- a/Engine/Source/Runtime/Sockets/Private/SocketSubsystem.cpp
+++ b/Engine/Source/Runtime/Sockets/Private/SocketSubsystem.cpp
@@ -5,6 +5,7 @@
 #include "Misc/ScopeLock.h"
 #include "SocketSubsystemModule.h"
 #include "Modules/ModuleManager.h"
+#include "SocketTypes.h"
 #include "IPAddress.h"
 #include "Sockets.h"
 #include "Templates/UniquePtr.h"
@@ -17,6 +18,12 @@ IMPLEMENT_MODULE( FSocketSubsystemModule, Sockets );
 extern FName CreateSocketSubsystem( FSocketSubsystemModule& SocketSubsystemModule );
 extern void DestroySocketSubsystem( FSocketSubsystemModule& SocketSubsystemModule );
 
+namespace FNetworkProtocolTypes
+{
+	const FName IPv4(TEXT("IPv4"));
+	const FName IPv6(TEXT("IPv6"));
+}
+
 /** Helper function to turn the friendly subsystem name into the module name */
 static inline FName GetSocketModuleName(const FString& SubsystemName)
 {
@@ -257,7 +264,7 @@ TSharedRef<FInternetAddr> ISocketSubsystem::GetLocalBindAddr(FOutputDevice& Out)
 	// look up the local host address
 	TSharedRef<FInternetAddr> BindAddr = GetLocalHostAddr(Out, bCanBindAll);
 
-	// If we can bind to all addresses, return 0.0.0.0
+	// If we can bind to all addresses, return the any address
 	if (bCanBindAll)
 	{
 		BindAddr->SetAnyAddress();
@@ -290,10 +297,8 @@ FResolveInfo* ISocketSubsystem::GetHostByName(const ANSICHAR* HostName)
 TSharedRef<FInternetAddr> ISocketSubsystem::GetLocalHostAddr(FOutputDevice& Out, bool& bCanBindAll)
 {
 	TSharedRef<FInternetAddr> HostAddr = CreateInternetAddr();
-	HostAddr->SetAnyAddress();
-
-	bCanBindAll = false;
 	FString HostName;
+	bCanBindAll = false;
 
 	if (GetMultihomeAddress(HostAddr))
 	{
@@ -306,7 +311,9 @@ TSharedRef<FInternetAddr> ISocketSubsystem::GetLocalHostAddr(FOutputDevice& Out,
 	}
 
 	// Failing to find the host is not considered an error and we just bind to any address
+	PRAGMA_DISABLE_DEPRECATION_WARNINGS
 	ESocketErrors FindHostResult = GetHostByName(TCHAR_TO_ANSI(*HostName), *HostAddr);
+	PRAGMA_ENABLE_DEPRECATION_WARNINGS
 	if (FindHostResult == SE_NO_ERROR || FindHostResult == SE_HOST_NOT_FOUND 
 		|| FindHostResult == SE_EWOULDBLOCK || FindHostResult == SE_TRY_AGAIN)
 	{
@@ -324,24 +331,28 @@ TSharedRef<FInternetAddr> ISocketSubsystem::GetLocalHostAddr(FOutputDevice& Out,
 	else
 	{
 		Out.Logf(TEXT("GetHostByName failed (%s)"), GetSocketError(FindHostResult));
+		HostAddr->SetAnyAddress();
 	}
 
 	// return the newly created address
 	return HostAddr;
 }
 
-bool ISocketSubsystem::GetMultihomeAddress(TSharedRef<class FInternetAddr>& Addr)
+bool ISocketSubsystem::GetMultihomeAddress(TSharedRef<FInternetAddr>& Addr)
 {
 	TCHAR Home[256] = TEXT("");
 	if (FParse::Value(FCommandLine::Get(), TEXT("MULTIHOME="), Home, ARRAY_COUNT(Home)))
 	{
-		bool bIsValid = false;
-		Addr->SetIp(Home, bIsValid);
-		if (Home == NULL || !bIsValid)
+		TSharedPtr<FInternetAddr> MultiHomeQuery = GetAddressFromString(Home);
+		if (Home == NULL || !MultiHomeQuery.IsValid())
 		{
 			UE_LOG(LogSockets, Log, TEXT("Invalid multihome IP address %s"), Home);
 			return false;
 		}
+		else
+		{
+			Addr = MultiHomeQuery.ToSharedRef();
+		}
 
 		return true;
 	}
@@ -375,6 +386,33 @@ void ISocketSubsystem::RemoveHostNameFromCache(const ANSICHAR* HostName)
 	HostNameCache.Remove(FString(HostName));
 }
 
+ESocketProtocolFamily ISocketSubsystem::GetProtocolFamilyFromName(const FName& InProtocolName) const
+{
+	if (InProtocolName == FNetworkProtocolTypes::IPv6)
+	{
+		return ESocketProtocolFamily::IPv6;
+	}
+	else if (InProtocolName == FNetworkProtocolTypes::IPv4)
+	{
+		return ESocketProtocolFamily::IPv4;
+	}
+
+	return ESocketProtocolFamily::None;
+}
+
+FName ISocketSubsystem::GetProtocolNameFromFamily(ESocketProtocolFamily InProtocolFamily) const
+{
+	switch (InProtocolFamily)
+	{
+	default:
+		return NAME_None;
+	case ESocketProtocolFamily::IPv4:
+		return FNetworkProtocolTypes::IPv4;
+	case ESocketProtocolFamily::IPv6:
+		return FNetworkProtocolTypes::IPv6;
+	}
+}
+
 FResolveInfoCached* ISocketSubsystem::CreateResolveInfoCached(TSharedPtr<FInternetAddr> Addr) const
 {
 	return new FResolveInfoCached(*Addr);
diff --git a/Engine/Source/Runtime/Sockets/Private/Unix/SocketSubsystemUnix.cpp b/Engine/Source/Runtime/Sockets/Private/Unix/SocketSubsystemUnix.cpp
--- a/Engine/Source/Runtime/Sockets/Private/Unix/SocketSubsystemUnix.cpp
+++ b/Engine/Source/Runtime/Sockets/Private/Unix/SocketSubsystemUnix.cpp
@@ -91,7 +91,7 @@ bool FSocketSubsystemUnix::HasNetworkDevice()
 	return true;
 }
 
-FSocket* FSocketSubsystemUnix::CreateSocket(const FName& SocketType, const FString& SocketDescription, ESocketProtocolFamily ProtocolType)
+FSocket* FSocketSubsystemUnix::CreateSocket(const FName& SocketType, const FString& SocketDescription, const FName& ProtocolType)
 {
 	FSocketBSD* NewSocket = (FSocketBSD*)FSocketSubsystemBSD::CreateSocket(SocketType, SocketDescription, ProtocolType);
 
@@ -172,7 +172,7 @@ bool FSocketSubsystemUnix::GetLocalAdapterAddresses(TArray<TSharedPtr<FInternetA
 				// Don't write it to our list again though.
 				if (bHasMultihome && NewAddress == MultihomeAddress)
 				{
-					static_cast<FInternetAddrBSD&>(MultihomeAddress.Get()).SetScopeId(AddressInterface);
+					StaticCastSharedRef<FInternetAddrBSD>(MultihomeAddress)->SetScopeId(AddressInterface);
 				}
 				else
 				{
diff --git a/Engine/Source/Runtime/Sockets/Private/Unix/SocketSubsystemUnix.h b/Engine/Source/Runtime/Sockets/Private/Unix/SocketSubsystemUnix.h
--- a/Engine/Source/Runtime/Sockets/Private/Unix/SocketSubsystemUnix.h
+++ b/Engine/Source/Runtime/Sockets/Private/Unix/SocketSubsystemUnix.h
@@ -50,7 +50,7 @@ public:
 	virtual bool Init(FString& Error) override;
 	virtual void Shutdown() override;
 	virtual bool HasNetworkDevice() override;
-	virtual class FSocket* CreateSocket(const FName& SocketType, const FString& SocketDescription, ESocketProtocolFamily ProtocolType) override;
+	virtual class FSocket* CreateSocket(const FName& SocketType, const FString& SocketDescription, const FName& ProtocolType) override;
 	virtual TSharedRef<FInternetAddr> GetLocalHostAddr(FOutputDevice& Out, bool& bCanBindAll) override;
 	virtual bool GetLocalAdapterAddresses(TArray<TSharedPtr<FInternetAddr> >& OutAddresses) override;
 };
diff --git a/Engine/Source/Runtime/Sockets/Private/Windows/SocketSubsystemWindows.cpp b/Engine/Source/Runtime/Sockets/Private/Windows/SocketSubsystemWindows.cpp
--- a/Engine/Source/Runtime/Sockets/Private/Windows/SocketSubsystemWindows.cpp
+++ b/Engine/Source/Runtime/Sockets/Private/Windows/SocketSubsystemWindows.cpp
@@ -69,7 +69,7 @@ void FSocketSubsystemWindows::Destroy()
 /* FSocketSubsystemBSD overrides
 *****************************************************************************/
 
-FSocket* FSocketSubsystemWindows::CreateSocket(const FName& SocketType, const FString& SocketDescription, ESocketProtocolFamily ProtocolType)
+FSocket* FSocketSubsystemWindows::CreateSocket(const FName& SocketType, const FString& SocketDescription, const FName& ProtocolType)
 {
 	FSocketBSD* NewSocket = (FSocketBSD*)FSocketSubsystemBSD::CreateSocket(SocketType, SocketDescription, ProtocolType);
 
@@ -132,7 +132,7 @@ ESocketErrors FSocketSubsystemWindows::GetLastErrorCode()
 
 bool FSocketSubsystemWindows::GetLocalAdapterAddresses( TArray<TSharedPtr<FInternetAddr> >& OutAdresses )
 {
-	ULONG Flags = GAA_FLAG_INCLUDE_PREFIX | GAA_FLAG_SKIP_MULTICAST | GAA_FLAG_SKIP_DNS_SERVER | GAA_FLAG_SKIP_FRIENDLY_NAME;
+	ULONG Flags = GAA_FLAG_INCLUDE_PREFIX | GAA_FLAG_SKIP_MULTICAST | GAA_FLAG_SKIP_ANYCAST | GAA_FLAG_SKIP_DNS_SERVER | GAA_FLAG_SKIP_FRIENDLY_NAME;
 	ULONG Result;
 	ULONG Size = 0;
 	ULONG Family = (PLATFORM_HAS_BSD_IPV6_SOCKETS) ? AF_UNSPEC : AF_INET;
@@ -170,7 +170,7 @@ bool FSocketSubsystemWindows::GetLocalAdapterAddresses( TArray<TSharedPtr<FInter
 					const sockaddr_storage* RawAddress = (const sockaddr_storage*)(UnicastAddress->Address.lpSockaddr);
 					TSharedRef<FInternetAddrBSD> NewAddress = MakeShareable(new FInternetAddrBSD(this));
 					NewAddress->SetIp(*RawAddress);
-					NewAddress->SetScopeId(ntohl(AdapterAddress->Ipv6IfIndex));
+					NewAddress->SetScopeId((RawAddress->ss_family == AF_INET) ? ntohl(AdapterAddress->IfIndex) : ntohl(AdapterAddress->Ipv6IfIndex));
 					OutAdresses.Add(NewAddress);
 				}
 			}
diff --git a/Engine/Source/Runtime/Sockets/Private/Windows/SocketSubsystemWindows.h b/Engine/Source/Runtime/Sockets/Private/Windows/SocketSubsystemWindows.h
--- a/Engine/Source/Runtime/Sockets/Private/Windows/SocketSubsystemWindows.h
+++ b/Engine/Source/Runtime/Sockets/Private/Windows/SocketSubsystemWindows.h
@@ -28,7 +28,7 @@ public:
 
 	// FSocketSubsystemBSD overrides
 
-	virtual class FSocket* CreateSocket(const FName& SocketType, const FString& SocketDescription, ESocketProtocolFamily ProtocolType) override;
+	virtual class FSocket* CreateSocket(const FName& SocketType, const FString& SocketDescription, const FName& ProtocolType) override;
 	virtual bool HasNetworkDevice() override;
 	virtual ESocketErrors GetLastErrorCode() override;
 	virtual bool GetLocalAdapterAddresses( TArray<TSharedPtr<FInternetAddr> >& OutAdresses ) override;
diff --git a/Engine/Source/Runtime/Sockets/Public/AddressInfoTypes.h b/Engine/Source/Runtime/Sockets/Public/AddressInfoTypes.h
--- a/Engine/Source/Runtime/Sockets/Public/AddressInfoTypes.h
+++ b/Engine/Source/Runtime/Sockets/Public/AddressInfoTypes.h
@@ -40,6 +40,7 @@ ENUM_CLASS_FLAGS(EAddressInfoFlags);
 
 struct FAddressInfoResultData
 {
+	UE_DEPRECATED(4.23, "Use constructor that supports an FName as the protocol specification as it supports multiple protocol types")
 	FAddressInfoResultData(TSharedRef<FInternetAddr> InAddr, SIZE_T InAddrLen, ESocketProtocolFamily InProtocol, ESocketType InSocketConfiguration) :
 		AddressProtocol(InProtocol),
 		SocketConfiguration(InSocketConfiguration),
@@ -48,6 +49,15 @@ struct FAddressInfoResultData
 	{
 	}
 
+	FAddressInfoResultData(TSharedRef<FInternetAddr> InAddr, SIZE_T InAddrLen, const FName InProtocol, ESocketType InSocketConfiguration) :
+		AddressProtocol(ESocketProtocolFamily::None),
+		AddressProtocolName(InProtocol),
+		SocketConfiguration(InSocketConfiguration),
+		AddressLen(InAddrLen),
+		Address(InAddr)
+	{
+	}
+
 	FName GetSocketTypeName() const
 	{
 		switch (SocketConfiguration)
@@ -65,8 +75,18 @@ struct FAddressInfoResultData
 		}
 	}
 
+	bool operator==(const FAddressInfoResultData& rhs) const
+	{
+		return SocketConfiguration == rhs.SocketConfiguration &&
+			AddressProtocolName == rhs.AddressProtocolName &&
+			*Address == *(rhs.Address);
+	}
+
 	/* The protocol of the address stored */
+	UE_DEPRECATED(4.23, "Use AddressProtocolName as it supports multiple protocol types")
 	ESocketProtocolFamily AddressProtocol;
+	/* The protocol name of the address stored */
+	FName AddressProtocolName;
 	/* Streaming or datagram */
 	ESocketType SocketConfiguration;
 	/* Length of the returned address data */
@@ -79,16 +99,19 @@ struct FAddressInfoResult
 {
 	FAddressInfoResult(const TCHAR* InHostName, const TCHAR* InServiceName) :
 		QueryHostName(InHostName),
-		QueryServiceName(InServiceName)
+		QueryServiceName(InServiceName),
+		ReturnCode(SE_NO_DATA)
 	{
 	}
 
-	/* The hostname that generated these results */
+	/* The hostname/ip that generated these results */
 	FString QueryHostName;
 	/* The service name that was used in the query */
 	FString QueryServiceName;
 	/* The Canonical Name of the query (empty unless FQDomainName or CanonicalName are specified) */
 	FString CanonicalNameResult;
+	/* The return code of the query */
+	ESocketErrors ReturnCode;
 	/* The list of results */
 	TArray<FAddressInfoResultData> Results;
 };
diff --git a/Engine/Source/Runtime/Sockets/Public/IPAddress.h b/Engine/Source/Runtime/Sockets/Public/IPAddress.h
--- a/Engine/Source/Runtime/Sockets/Public/IPAddress.h
+++ b/Engine/Source/Runtime/Sockets/Public/IPAddress.h
@@ -5,6 +5,7 @@
 #include "CoreMinimal.h"
 #include "Stats/Stats.h"
 #include "Async/AsyncWork.h"
+#include "SocketTypes.h"
 
 /**
  * Represents an internet address. All data is in network byte order
@@ -144,7 +145,7 @@ public:
 	 *
 	 * @return	The value to use for the hash
 	 */
-	virtual uint32 GetTypeHash() = 0;
+	virtual uint32 GetTypeHash() const = 0;
 
 	/**
 	 * Is this a well formed internet address
@@ -159,6 +160,17 @@ public:
 	 * @return The new structure
 	 */
 	virtual TSharedRef<FInternetAddr> Clone() const = 0;
+
+	/**
+	 * Returns the protocol type name of the address data currently stored in this struct
+	 *
+	 * @return The type of the address. 
+	 *         If it's not known or overridden in a derived class, the return is None.
+	 */
+	virtual FName GetProtocolType() const
+	{
+		return NAME_None;
+	}
 };
 
 /**
diff --git a/Engine/Source/Runtime/Sockets/Public/SocketSubsystem.h b/Engine/Source/Runtime/Sockets/Public/SocketSubsystem.h
--- a/Engine/Source/Runtime/Sockets/Public/SocketSubsystem.h
+++ b/Engine/Source/Runtime/Sockets/Public/SocketSubsystem.h
@@ -4,6 +4,7 @@
 
 #include "CoreMinimal.h"
 #include "SocketTypes.h"
+#include "IPAddress.h"
 #include "AddressInfoTypes.h"
 
 class Error;
@@ -84,7 +85,7 @@ public:
 	 */
 	virtual class FSocket* CreateSocket(const FName& SocketType, const FString& SocketDescription, bool bForceUDP = false)
 	{
-		return CreateSocket(SocketType, SocketDescription, ESocketProtocolFamily::None);
+		return CreateSocket(SocketType, SocketDescription, NAME_None);
 	}
 
 	/**
@@ -92,11 +93,26 @@ public:
 	 *
 	 * @Param SocketType type of socket to create (DGram, Stream, etc)
 	 * @param SocketDescription debug description
-	 * @param ProtocolType the socket protocol to be used
+	 * @param ProtocolType the socket protocol to be used. Each subsystem must handle the None case and output a valid socket regardless.
 	 *
 	 * @return the new socket or NULL if failed
 	 */
-	virtual class FSocket* CreateSocket(const FName& SocketType, const FString& SocketDescription, ESocketProtocolFamily ProtocolType) = 0;
+	UE_DEPRECATED(4.23, "Use the CreateSocket with the FName parameter for support for multiple protocol types.")
+	virtual class FSocket* CreateSocket(const FName& SocketType, const FString& SocketDescription, ESocketProtocolFamily ProtocolType)
+	{
+		return CreateSocket(SocketType, SocketDescription, GetProtocolNameFromFamily(ProtocolType));
+	}
+
+	/**
+	 * Creates a socket using the given protocol name.
+	 *
+	 * @Param SocketType type of socket to create (DGram, Stream, etc)
+	 * @param SocketDescription debug description
+	 * @param ProtocolName the name of the internet protocol to use for this socket. None should be handled.
+	 *
+	 * @return the new socket or NULL if failed
+	 */
+	virtual class FSocket* CreateSocket(const FName& SocketType, const FString& SocketDescription, const FName& ProtocolName) = 0;
 
 	/**
 	 * Creates a resolve info cached struct to hold the resolved address
@@ -130,18 +146,71 @@ public:
 	 *
 	 * @return the results structure containing the array of address results to this query
 	 */
+	UE_DEPRECATED(4.23, "Migrate to GetAddressInfo that takes an FName as the protocol specification.")
 	virtual FAddressInfoResult GetAddressInfo(const TCHAR* HostName, const TCHAR* ServiceName = nullptr,
 		EAddressInfoFlags QueryFlags = EAddressInfoFlags::Default,
 		ESocketProtocolFamily ProtocolType = ESocketProtocolFamily::None,
+		ESocketType SocketType = ESocketType::SOCKTYPE_Unknown)
+	{
+		return GetAddressInfo(HostName, ServiceName, QueryFlags, GetProtocolNameFromFamily(ProtocolType), SocketType);
+	}
+
+	/**
+	 * Gets the address information of the given hostname and outputs it into an array of resolvable addresses.
+	 * It is up to the caller to determine which one is valid for their environment.
+	 *
+	 * This function allows for specifying FNames for the protocol type, allowing for support of other 
+	 * platform protocols
+	 *
+	 * @param HostName string version of the queryable hostname or ip address
+	 * @param ServiceName string version of a service name ("http") or a port number ("80")
+	 * @param QueryFlags What flags are used in making the getaddrinfo call. Several flags can be used at once by
+	 *                   bitwise OR-ing the flags together.
+	 *                   Platforms are required to translate this value into a the correct flag representation.
+	 * @param ProtocolTypeName Used to limit results from the call. Specifying None will search all valid protocols.
+	 *					   Callers will find they rarely have to specify this flag.
+	 * @param SocketType What socket type should the results be formatted for. This typically does not change any
+	 *                   formatting results and can be safely left to the default value.
+	 *
+	 * @return the results structure containing the array of address results to this query
+	 */
+	virtual FAddressInfoResult GetAddressInfo(const TCHAR* HostName, const TCHAR* ServiceName = nullptr,
+		EAddressInfoFlags QueryFlags = EAddressInfoFlags::Default,
+		const FName ProtocolTypeName = NAME_None,
 		ESocketType SocketType = ESocketType::SOCKTYPE_Unknown) = 0;
 
+	/**
+	 * Serializes a string that only contains an address.
+	 * 
+	 * This is a what you see is what you get, there is no DNS resolution of the input string,
+	 * so only use this if you know you already have a valid address.
+	 * Otherwise, feed the address to GetAddressInfo for guaranteed results.
+	 *
+	 * @param InAddress the address to serialize
+	 *
+	 * @return The FInternetAddr of the given string address. This will point to nullptr on failure.
+	 */
+	virtual TSharedPtr<FInternetAddr> GetAddressFromString(const FString& InAddress) = 0;
+
 	/**
 	 * Does a DNS look up of a host name
+	 * This code assumes a lot, and as such, it's not guaranteed that the results provided by it are correct.
 	 *
 	 * @param HostName the name of the host to look up
 	 * @param Addr the address to copy the IP address to
 	 */
-	virtual ESocketErrors GetHostByName(const ANSICHAR* HostName, FInternetAddr& OutAddr) = 0;
+	UE_DEPRECATED(4.23, "Please use GetAddressInfo to query hostnames")
+	virtual ESocketErrors GetHostByName(const ANSICHAR* HostName, FInternetAddr& OutAddr)
+	{
+		FAddressInfoResult GAIResult = GetAddressInfo(ANSI_TO_TCHAR(HostName), nullptr, EAddressInfoFlags::Default, NAME_None);
+		if (GAIResult.Results.Num() > 0)
+		{
+			OutAddr.SetRawIp(GAIResult.Results[0].Address->GetRawIp());
+			return SE_NO_ERROR;
+		}
+
+		return SE_HOST_NOT_FOUND;
+	}
 
 	/**
 	 * Creates a platform specific async hostname resolution object
@@ -179,7 +248,21 @@ public:
 	 * @param Address host address
 	 * @param Port host port
 	 */
-	virtual TSharedRef<FInternetAddr> CreateInternetAddr(uint32 Address=0, uint32 Port=0) = 0;
+	UE_DEPRECATED(4.23, "To support different address sizes, use CreateInternetAddr with no arguments and call SetIp/SetRawIp and SetPort on the returned object")
+	virtual TSharedRef<FInternetAddr> CreateInternetAddr(uint32 Address, uint32 Port = 0)
+	{
+		TSharedRef<FInternetAddr> ReturnAddr = CreateInternetAddr();
+		PRAGMA_DISABLE_DEPRECATION_WARNINGS
+		ReturnAddr->SetIp(Address);
+		ReturnAddr->SetPort(Port);
+		PRAGMA_ENABLE_DEPRECATION_WARNINGS
+		return ReturnAddr;
+	}
+
+	/**
+	 * Create a proper FInternetAddr representation
+	 */
+	virtual TSharedRef<FInternetAddr> CreateInternetAddr() = 0;
 
 	/**
 	 * @return Whether the machine has a properly configured network device or not
@@ -254,9 +337,9 @@ public:
 	 * 
 	 * @param Addr the address structure which will have the Multihome address in it if set.
 	 *
-	 * @return If the multihome address is valid or not.
+	 * @return If the multihome address was set and valid
 	 */
-	virtual bool GetMultihomeAddress(TSharedRef<class FInternetAddr>& Addr);
+	virtual bool GetMultihomeAddress(TSharedRef<FInternetAddr>& Addr);
 
 	/**
 	 * Checks the host name cache for an existing entry (faster than resolving again)
@@ -288,6 +371,15 @@ public:
 	 */
 	virtual bool IsSocketWaitSupported() const = 0;
 
+protected:
+
+	/**
+	 * Conversion functions from the SocketProtocolFamily enum to the new FName system.
+	 * For now, both are supported, but it's better to use the FName when possible.
+	 */
+	virtual ESocketProtocolFamily GetProtocolFamilyFromName(const FName& InProtocolName) const;
+	virtual FName GetProtocolNameFromFamily(ESocketProtocolFamily InProtocolFamily) const;
+
 private:
 
 	/** Used to provide threadsafe access to the host name cache */
diff --git a/Engine/Source/Runtime/Sockets/Public/SocketTypes.h b/Engine/Source/Runtime/Sockets/Public/SocketTypes.h
--- a/Engine/Source/Runtime/Sockets/Public/SocketTypes.h
+++ b/Engine/Source/Runtime/Sockets/Public/SocketTypes.h
@@ -14,6 +14,13 @@ enum class ESocketProtocolFamily : uint8
 	IPv6
 };
 
+/** Used for indicating the socket network protocol to be used in addressing and socket creation */
+namespace FNetworkProtocolTypes
+{
+	SOCKETS_API extern const FName IPv4;
+	SOCKETS_API extern const FName IPv6;
+}
+
 /** Indicates the type of socket being used (streaming or datagram) */
 enum ESocketType
 {
diff --git a/Engine/Source/Runtime/Sockets/Public/Sockets.h b/Engine/Source/Runtime/Sockets/Public/Sockets.h
--- a/Engine/Source/Runtime/Sockets/Public/Sockets.h
+++ b/Engine/Source/Runtime/Sockets/Public/Sockets.h
@@ -21,7 +21,7 @@ protected:
 	FString SocketDescription;
 
 	/** Protocol used in creation of a socket */
-	ESocketProtocolFamily SocketProtocol;
+	FName SocketProtocol;
 
 public:
 
@@ -29,7 +29,7 @@ public:
 	inline FSocket() :
 		SocketType(SOCKTYPE_Unknown),
 		SocketDescription(TEXT("")),
-		SocketProtocol(ESocketProtocolFamily::None)
+		SocketProtocol(NAME_None)
 	{ }
 
 	/**
@@ -42,9 +42,28 @@ public:
 	inline FSocket(ESocketType InSocketType, const FString& InSocketDescription) :
 		SocketType(InSocketType),
 		SocketDescription(InSocketDescription),
-		SocketProtocol(ESocketProtocolFamily::None)
+		SocketProtocol(NAME_None)
 	{ }
 
+	UE_DEPRECATED(4.23, "Please migrate to the constructor that takes an FName for the protocol stack")
+		inline FSocket(ESocketType InSocketType, const FString& InSocketDescription, ESocketProtocolFamily InSocketProtocol) :
+		SocketType(InSocketType),
+		SocketDescription(InSocketDescription)
+	{ 
+		if (InSocketProtocol == ESocketProtocolFamily::IPv4)
+		{
+			SocketProtocol = FNetworkProtocolTypes::IPv4;
+		}
+		else if (InSocketProtocol == ESocketProtocolFamily::IPv6)
+		{
+			SocketProtocol = FNetworkProtocolTypes::IPv6;
+		}
+		else
+		{
+			SocketProtocol = NAME_None;
+		}
+	}
+
 	/**
 	 * Specifies the type of socket being created
 	 *
@@ -52,7 +71,7 @@ public:
 	 * @param InSocketDescription The debug description of the socket
 	 * @param InSocketProtocol the protocol stack this socket should be created on.
 	 */
-	inline FSocket(ESocketType InSocketType, const FString& InSocketDescription, ESocketProtocolFamily InSocketProtocol) :
+	inline FSocket(ESocketType InSocketType, const FString& InSocketDescription, const FName& InSocketProtocol) :
 		SocketType(InSocketType),
 		SocketDescription(InSocketDescription),
 		SocketProtocol(InSocketProtocol)
@@ -402,7 +421,7 @@ public:
 	 * @return Socket type.
 	 * @see GetDescription, GetPortNo
 	 */
-	FORCEINLINE ESocketProtocolFamily GetProtocol() const
+	FORCEINLINE FName GetProtocol() const
 	{
 		return SocketProtocol;
 	}
