Fix Linux build break due to Clang being unable to determine which override function to use.

#rnx

diff --git a/Engine/Source/Runtime/Sockets/Public/SocketSubsystem.h b/Engine/Source/Runtime/Sockets/Public/SocketSubsystem.h
--- a/Engine/Source/Runtime/Sockets/Public/SocketSubsystem.h
+++ b/Engine/Source/Runtime/Sockets/Public/SocketSubsystem.h
@@ -85,7 +85,8 @@ public:
 	 */
 	virtual class FSocket* CreateSocket(const FName& SocketType, const FString& SocketDescription, bool bForceUDP = false)
 	{
-		return CreateSocket(SocketType, SocketDescription, NAME_None);
+		const FName NoProtocolTypeName(NAME_None);
+		return CreateSocket(SocketType, SocketDescription, NoProtocolTypeName);
 	}
 
 	/**
