Fix iOS build break due to missing cast qualifier

#rnx

diff --git a/Engine/Source/Runtime/Sockets/Private/IOS/IPAddressBSDIOS.h b/Engine/Source/Runtime/Sockets/Private/IOS/IPAddressBSDIOS.h
--- a/Engine/Source/Runtime/Sockets/Private/IOS/IPAddressBSDIOS.h
+++ b/Engine/Source/Runtime/Sockets/Private/IOS/IPAddressBSDIOS.h
@@ -21,7 +21,7 @@ public:
 		FSocketSubsystemBSD* SocketSubsystem = static_cast<FSocketSubsystemBSD*>(ISocketSubsystem::Get(PLATFORM_SOCKETSUBSYSTEM));
 		if (SocketSubsystem)
 		{
-			TSharedPtr<FInternetAddrBSD> MultiCastAddr = SocketSubsystem->GetAddressFromString(TEXT("ff02::01"));
+			TSharedPtr<FInternetAddrBSD> MultiCastAddr = StaticCastSharedPtr<FInternetAddrBSD>(SocketSubsystem->GetAddressFromString(TEXT("ff02::01")));
 			if (!MultiCastAddr.IsValid())
 			{
 				UE_LOG(LogSockets, Warning, TEXT("Could not resolve the broadcast address for iOS, this address will just be blank"));
