Fix bad resolution of invalid protocols when using FNames in SocketSubsystemBSD due to typo. This only affects hybrid stacks when we don't have enough information to determine protocol.

#rnx

diff --git a/Engine/Source/Runtime/Sockets/Private/BSDSockets/SocketSubsystemBSD.cpp b/Engine/Source/Runtime/Sockets/Private/BSDSockets/SocketSubsystemBSD.cpp
--- a/Engine/Source/Runtime/Sockets/Private/BSDSockets/SocketSubsystemBSD.cpp
+++ b/Engine/Source/Runtime/Sockets/Private/BSDSockets/SocketSubsystemBSD.cpp
@@ -26,16 +26,15 @@ FSocket* FSocketSubsystemBSD::CreateSocket(const FName& SocketType, const FStrin
 	// For platforms that have two subsystems (ex: Steam) but don't explicitly inherit from SocketSubsystemBSD
 	// so they don't know which protocol to end up using and pass in None.
 	// This is invalid, so we need to attempt to still resolve it.
-	if (SocketProtocolType != FNetworkProtocolTypes::IPv4 || SocketProtocolType != FNetworkProtocolTypes::IPv6)
+	if (ProtocolType != FNetworkProtocolTypes::IPv4 && ProtocolType != FNetworkProtocolTypes::IPv6)
 	{
 		SocketProtocolType = GetDefaultSocketProtocolFamily();
-	}
-
-	// Don't support any other protocol families.
-	if (SocketProtocolType != FNetworkProtocolTypes::IPv4 && SocketProtocolType != FNetworkProtocolTypes::IPv6)
-	{
-		UE_LOG(LogSockets, Warning, TEXT("Provided socket protocol type is unsupported! Returning null socket"));
-		return nullptr;
+		// Check to make sure this is still valid.
+		if (SocketProtocolType != FNetworkProtocolTypes::IPv4 && SocketProtocolType != FNetworkProtocolTypes::IPv6)
+		{
+			UE_LOG(LogSockets, Warning, TEXT("Provided socket protocol type is unsupported! Returning null socket"));
+			return nullptr;
+		}
 	}
 
 #if PLATFORM_HAS_BSD_SOCKET_FEATURE_CLOSE_ON_EXEC
