Updated log verbosity of getaddrinfo results so they will log by default, and added port information.

diff --git a/Engine/Source/Runtime/Sockets/Private/BSDSockets/SocketSubsystemBSD.cpp b/Engine/Source/Runtime/Sockets/Private/BSDSockets/SocketSubsystemBSD.cpp
--- a/Engine/Source/Runtime/Sockets/Private/BSDSockets/SocketSubsystemBSD.cpp
+++ b/Engine/Source/Runtime/Sockets/Private/BSDSockets/SocketSubsystemBSD.cpp
@@ -110,7 +110,7 @@ FAddressInfoResult FSocketSubsystemBSD::GetAddressInfo(const TCHAR* HostName, co
 	int32 ErrorCode = getaddrinfo(TCHAR_TO_UTF8(HostName), TCHAR_TO_UTF8(ServiceName), &HintAddrInfo, &AddrInfo);
 	AddrQueryResult.ReturnCode = TranslateGAIErrorCode(ErrorCode);
 
-	UE_LOG(LogSockets, Verbose, TEXT("Executed getaddrinfo with HostName: %s Return: %d"), HostName, ErrorCode);
+	UE_LOG(LogSockets, Log, TEXT("Executed getaddrinfo with HostName: %s ServiceName: %s Return: %d"), HostName ? HostName : TEXT("null"), ServiceName ? ServiceName : TEXT("null"), ErrorCode);
 	if (AddrQueryResult.ReturnCode == SE_NO_ERROR)
 	{
 		addrinfo* AddrInfoHead = AddrInfo;
@@ -134,7 +134,7 @@ FAddressInfoResult FSocketSubsystemBSD::GetAddressInfo(const TCHAR* HostName, co
 					if (AddrQueryResult.Results.Add(NewAddressData) != INDEX_NONE)
 					{
 						sockaddr_in6* IPv6AddrData = reinterpret_cast<sockaddr_in6*>(AddrInfo->ai_addr);
-						UE_LOG(LogSockets, Verbose, TEXT("# Family: %s Address: %s And scope %d"), ((AddrInfo->ai_family == AF_INET) ? TEXT("IPv4") : TEXT("IPv6")), *(NewAddress->ToString(false)), ((AddrInfo->ai_family == AF_INET) ? -1 : IPv6AddrData->sin6_scope_id));
+						UE_LOG(LogSockets, Log, TEXT("# Family: %s Address: %s And scope %d"), ((AddrInfo->ai_family == AF_INET) ? TEXT("IPv4") : TEXT("IPv6")), *(NewAddress->ToString(true)), ((AddrInfo->ai_family == AF_INET) ? -1 : IPv6AddrData->sin6_scope_id));
 					}
 				}
 			}
