Fill out the enum data about addresses for GAI results for backwards compatibility.

#rnx

diff --git a/Engine/Source/Runtime/Sockets/Private/BSDSockets/SocketSubsystemBSD.cpp b/Engine/Source/Runtime/Sockets/Private/BSDSockets/SocketSubsystemBSD.cpp
--- a/Engine/Source/Runtime/Sockets/Private/BSDSockets/SocketSubsystemBSD.cpp
+++ b/Engine/Source/Runtime/Sockets/Private/BSDSockets/SocketSubsystemBSD.cpp
@@ -131,6 +131,12 @@ FAddressInfoResult FSocketSubsystemBSD::GetAddressInfo(const TCHAR* HostName, co
 					NewAddress->Set(*AddrData, AddrInfo->ai_addrlen);
 					
 					FAddressInfoResultData NewAddressData(NewAddress, AddrInfo->ai_addrlen, GetProtocolFamilyTypeName(AddrInfo->ai_family), GetSocketType(AddrInfo->ai_protocol));
+
+					// While the FNames are the correct way to enumerate over address types, we still need to support the old format.
+					PRAGMA_DISABLE_DEPRECATION_WARNINGS
+					NewAddressData.AddressProtocol = GetProtocolFamilyType(AddrInfo->ai_family);
+					PRAGMA_ENABLE_DEPRECATION_WARNINGS
+
 					if (AddrQueryResult.Results.Add(NewAddressData) != INDEX_NONE)
 					{
 						sockaddr_in6* IPv6AddrData = reinterpret_cast<sockaddr_in6*>(AddrInfo->ai_addr);
