Print out the network error code and message when an address serialization failure occurs.

#rnx

diff --git a/Engine/Source/Runtime/Sockets/Private/BSDIPv6Sockets/SocketSubsystemBSDIPv6.cpp b/Engine/Source/Runtime/Sockets/Private/BSDIPv6Sockets/SocketSubsystemBSDIPv6.cpp
--- a/Engine/Source/Runtime/Sockets/Private/BSDIPv6Sockets/SocketSubsystemBSDIPv6.cpp
+++ b/Engine/Source/Runtime/Sockets/Private/BSDIPv6Sockets/SocketSubsystemBSDIPv6.cpp
@@ -208,7 +208,8 @@ TSharedPtr<FInternetAddr> FSocketSubsystemBSDIPv6::GetAddressFromString(const FS
 		return ReturnAddress;
 	}
 
-	UE_LOG(LogSockets, Warning, TEXT("Could not serialize %s, got error code %d"), *InAddress, GetLastErrorCode());
+	ESocketErrors LastError = GetLastErrorCode();
+	UE_LOG(LogSockets, Warning, TEXT("Could not serialize %s, got error code %s [%d]"), *InAddress, GetSocketError(LastError), LastError);
 	return nullptr;
 }
 
diff --git a/Engine/Source/Runtime/Sockets/Private/BSDSockets/SocketSubsystemBSD.cpp b/Engine/Source/Runtime/Sockets/Private/BSDSockets/SocketSubsystemBSD.cpp
--- a/Engine/Source/Runtime/Sockets/Private/BSDSockets/SocketSubsystemBSD.cpp
+++ b/Engine/Source/Runtime/Sockets/Private/BSDSockets/SocketSubsystemBSD.cpp
@@ -187,7 +187,8 @@ TSharedPtr<FInternetAddr> FSocketSubsystemBSD::GetAddressFromString(const FStrin
 		return ReturnAddress;
 	}
 	
-	UE_LOG(LogSockets, Warning, TEXT("Could not serialize %s, got error code %d"), *InAddress, GetLastErrorCode());
+	ESocketErrors LastError = GetLastErrorCode();
+	UE_LOG(LogSockets, Warning, TEXT("Could not serialize %s, got error code %s [%d]"), *InAddress, GetSocketError(LastError), LastError);
 	return nullptr;
 }
 
