Adding socket error handling to prevent hang when requesting connection

diff --git a/Engine/Source/Runtime/Sockets/Private/BSDSockets/SocketSubsystemBSD.cpp b/Engine/Source/Runtime/Sockets/Private/BSDSockets/SocketSubsystemBSD.cpp
--- a/Engine/Source/Runtime/Sockets/Private/BSDSockets/SocketSubsystemBSD.cpp
+++ b/Engine/Source/Runtime/Sockets/Private/BSDSockets/SocketSubsystemBSD.cpp
@@ -319,7 +319,7 @@ ESocketErrors FSocketSubsystemBSD::TranslateErrorCode(int32 Code)
 // 	case SYSNOTREADY: return SE_SYSNOTREADY;
 // 	case VERNOTSUPPORTED: return SE_VERNOTSUPPORTED;
 // 	case NOTINITIALISED: return SE_NOTINITIALISED;
-
+	case EPIPE: return SE_ECONNRESET; // for when backgrounding with an open pipe to a server
 #if PLATFORM_HAS_BSD_SOCKET_FEATURE_GETHOSTNAME
 	case HOST_NOT_FOUND: return SE_HOST_NOT_FOUND;
 	case TRY_AGAIN: return SE_TRY_AGAIN;
