Added Linux FSocketUnix socket, with RecvMulti function for receiving multiple packets in one call - and SetRetrieveTimestamp, for receiving packet timestamp data from the operating system.

diff --git a/Engine/Source/Runtime/Core/Public/HAL/Platform.h b/Engine/Source/Runtime/Core/Public/HAL/Platform.h
--- a/Engine/Source/Runtime/Core/Public/HAL/Platform.h
+++ b/Engine/Source/Runtime/Core/Public/HAL/Platform.h
@@ -284,6 +284,12 @@
 #ifndef PLATFORM_HAS_BSD_SOCKET_FEATURE_MSG_DONTWAIT
 	#define PLATFORM_HAS_BSD_SOCKET_FEATURE_MSG_DONTWAIT	0
 #endif
+#ifndef PLATFORM_HAS_BSD_SOCKET_FEATURE_RECVMMSG
+	#define PLATFORM_HAS_BSD_SOCKET_FEATURE_RECVMMSG	0
+#endif
+#ifndef PLATFORM_HAS_BSD_SOCKET_FEATURE_TIMESTAMP
+	#define PLATFORM_HAS_BSD_SOCKET_FEATURE_TIMESTAMP 0
+#endif
 #ifndef PLATFORM_HAS_NO_EPROCLIM
 	#define PLATFORM_HAS_NO_EPROCLIM			0
 #endif
diff --git a/Engine/Source/Runtime/Core/Public/Unix/UnixPlatform.h b/Engine/Source/Runtime/Core/Public/Unix/UnixPlatform.h
--- a/Engine/Source/Runtime/Core/Public/Unix/UnixPlatform.h
+++ b/Engine/Source/Runtime/Core/Public/Unix/UnixPlatform.h
@@ -56,6 +56,8 @@ typedef FUnixPlatformTypes FPlatformTypes;
 #define PLATFORM_HAS_NO_EPROCLIM						1
 #define PLATFORM_HAS_BSD_SOCKET_FEATURE_IOCTL			1
 #define PLATFORM_HAS_BSD_SOCKET_FEATURE_MSG_DONTWAIT	1
+#define PLATFORM_HAS_BSD_SOCKET_FEATURE_RECVMMSG		1
+#define PLATFORM_HAS_BSD_SOCKET_FEATURE_TIMESTAMP		1
 #define PLATFORM_SUPPORTS_STACK_SYMBOLS					1
 #define PLATFORM_IS_ANSI_MALLOC_THREADSAFE				1
 #define PLATFORM_ALLOW_ALLOCATIONS_IN_FASYNCWRITER_SERIALIZEBUFFERTOARCHIVE 0
diff --git a/Engine/Source/Runtime/Sockets/Private/BSDSockets/IPAddressBSD.h b/Engine/Source/Runtime/Sockets/Private/BSDSockets/IPAddressBSD.h
--- a/Engine/Source/Runtime/Sockets/Private/BSDSockets/IPAddressBSD.h
+++ b/Engine/Source/Runtime/Sockets/Private/BSDSockets/IPAddressBSD.h
@@ -170,6 +170,11 @@ public:
 	 */
 	virtual void GetIp(uint32& OutAddr) const override;
 
+	sockaddr_storage* GetRawAddr()
+	{
+		return &Addr;
+	}
+
 	/**
 	 * Sets the port number from a host byte order int
 	 *
diff --git a/Engine/Source/Runtime/Sockets/Private/SocketSubsystem.cpp b/Engine/Source/Runtime/Sockets/Private/SocketSubsystem.cpp
--- a/Engine/Source/Runtime/Sockets/Private/SocketSubsystem.cpp
+++ b/Engine/Source/Runtime/Sockets/Private/SocketSubsystem.cpp
@@ -491,25 +491,24 @@ const TCHAR* ISocketSubsystem::GetSocketError(ESocketErrors Code)
 	};
 }
 
+TUniquePtr<FRecvMulti> ISocketSubsystem::CreateRecvMulti(int32 MaxNumPackets, int32 MaxPacketSize,
+															ERecvMultiFlags Flags/*=ERecvMultiFlags::None*/)
+{
+	UE_LOG(LogSockets, Warning, TEXT("RecvMulti is not supported by current socket subsystem."));
 
+	return nullptr;
+}
 
+bool ISocketSubsystem::IsSocketRecvMultiSupported() const
+{
+	return false;
+}
 
+double ISocketSubsystem::TranslatePacketTimestamp(const FPacketTimestamp& Timestamp,
+													ETimestampTranslation Translation/*=ETimestampTranslation::LocalTimestamp*/)
+{
+	UE_LOG(LogSockets, Warning, TEXT("TranslatePacketTimestamp is not supported by current socket subsystem."));
 
-
-
-
-
-
-
-
-
-
-
-
-
-
-
-
-
-
+	return 0.0;
+}
 
diff --git a/Engine/Source/Runtime/Sockets/Private/Sockets.cpp b/Engine/Source/Runtime/Sockets/Private/Sockets.cpp
--- a/Engine/Source/Runtime/Sockets/Private/Sockets.cpp
+++ b/Engine/Source/Runtime/Sockets/Private/Sockets.cpp
@@ -4,6 +4,27 @@
 #include "SocketSubsystem.h"
 //#include "Net/NetworkProfiler.h"
 
+
+/**
+ * FRecvMulti
+ */
+
+FRecvMulti::FRecvMulti(ISocketSubsystem* SocketSubsystem, int32 InMaxNumPackets, int32 InMaxPacketSize, ERecvMultiFlags InitFlags)
+	: Packets(MakeUnique<FRecvData[]>(InMaxNumPackets))
+	, NumPackets(0)
+	, MaxNumPackets(InMaxNumPackets)
+	, MaxPacketSize(InMaxPacketSize)
+{
+}
+
+void FRecvMulti::CountBytes(FArchive& Ar) const
+{
+	Ar.CountBytes(sizeof(this), sizeof(this));
+
+	// Packets
+	Ar.CountBytes(MaxNumPackets * sizeof(FRecvData), MaxNumPackets * sizeof(FRecvData));
+}
+
 //
 // FSocket stats implementation
 //
@@ -42,3 +63,13 @@ bool FSocket::Recv(uint8* Data, int32 BufferSize, int32& BytesRead, ESocketRecei
 	}
 	return true;
 }
+
+bool FSocket::RecvMulti(FRecvMulti& MultiData, ESocketReceiveFlags::Type Flags/*=ESocketReceiveFlags::None*/)
+{
+	return false;
+}
+
+bool FSocket::SetRetrieveTimestamp(bool bRetrieveTimestamp/*=true*/)
+{
+	return false;
+}
diff --git a/Engine/Source/Runtime/Sockets/Private/Unix/SocketSubsystemUnix.cpp b/Engine/Source/Runtime/Sockets/Private/Unix/SocketSubsystemUnix.cpp
--- a/Engine/Source/Runtime/Sockets/Private/Unix/SocketSubsystemUnix.cpp
+++ b/Engine/Source/Runtime/Sockets/Private/Unix/SocketSubsystemUnix.cpp
@@ -4,7 +4,7 @@
 #include "Misc/CommandLine.h"
 #include "SocketSubsystemModule.h"
 #include "BSDSockets/IPAddressBSD.h"
-#include "BSDSockets/SocketsBSD.h"
+#include "Unix/SocketsUnix.h"
 
 #include <ifaddrs.h>
 #include <net/if.h>
@@ -193,3 +193,59 @@ bool FSocketSubsystemUnix::GetLocalAdapterAddresses(TArray<TSharedPtr<FInternetA
 
 	return (OutAddresses.Num() > 0);
 }
+
+FSocketBSD* FSocketSubsystemUnix::InternalBSDSocketFactory(SOCKET Socket, ESocketType SocketType, const FString& SocketDescription, const FName& SocketProtocol)
+{
+	return new FSocketUnix(Socket, SocketType, SocketDescription, SocketProtocol, this);
+}
+
+TUniquePtr<FRecvMulti> FSocketSubsystemUnix::CreateRecvMulti(int32 MaxNumPackets, int32 MaxPacketSize, ERecvMultiFlags Flags)
+{
+#if PLATFORM_HAS_BSD_SOCKET_FEATURE_RECVMMSG
+	return MakeUnique<FUnixRecvMulti>(this, MaxNumPackets, MaxPacketSize, Flags);
+#endif
+
+	return nullptr;
+}
+
+bool FSocketSubsystemUnix::IsSocketRecvMultiSupported() const
+{
+#if PLATFORM_HAS_BSD_SOCKET_FEATURE_RECVMMSG
+	return true;
+#endif
+
+	return false;
+}
+
+double FSocketSubsystemUnix::TranslatePacketTimestamp(const FPacketTimestamp& Timestamp, ETimestampTranslation Translation)
+{
+	double ReturnVal = 0.0;
+	const bool bDeltaOnly = Translation == ETimestampTranslation::TimeDelta;
+
+	if (Translation == ETimestampTranslation::LocalTimestamp || bDeltaOnly)
+	{
+		// Unfortunately, the packet timestamp is platform-specific, using CLOCK_REALTIME in this case,
+		// whereas FPlatformTime may select from a variety of incompatible clocks.
+		// So, the only safe option is to return the time difference instead.
+		struct timespec CurPlatformTime;
+
+		clock_gettime(CLOCK_REALTIME, &CurPlatformTime);
+
+		FTimespan CurPlatformTimespan((CurPlatformTime.tv_sec * ETimespan::TicksPerSecond) +
+										(CurPlatformTime.tv_nsec / ETimespan::NanosecondsPerTick));
+
+		ReturnVal = (CurPlatformTimespan - Timestamp.Timestamp).GetTotalSeconds();
+
+		if (!bDeltaOnly)
+		{
+			ReturnVal = FPlatformTime::Seconds() - ReturnVal;
+		}
+	}
+	else
+	{
+		UE_LOG(LogSockets, Warning, TEXT("TranslatePacketTimestamp: Unknown ETimestampTranslation type: %i"), (uint32)Translation);
+	}
+
+	return ReturnVal;
+}
+
diff --git a/Engine/Source/Runtime/Sockets/Private/Unix/SocketSubsystemUnix.h b/Engine/Source/Runtime/Sockets/Private/Unix/SocketSubsystemUnix.h
--- a/Engine/Source/Runtime/Sockets/Private/Unix/SocketSubsystemUnix.h
+++ b/Engine/Source/Runtime/Sockets/Private/Unix/SocketSubsystemUnix.h
@@ -53,4 +53,8 @@ public:
 	virtual class FSocket* CreateSocket(const FName& SocketType, const FString& SocketDescription, const FName& ProtocolType) override;
 	virtual TSharedRef<FInternetAddr> GetLocalHostAddr(FOutputDevice& Out, bool& bCanBindAll) override;
 	virtual bool GetLocalAdapterAddresses(TArray<TSharedPtr<FInternetAddr> >& OutAddresses) override;
+	virtual class FSocketBSD* InternalBSDSocketFactory( SOCKET Socket, ESocketType SocketType, const FString& SocketDescription, const FName& SocketProtocol) override;
+	virtual TUniquePtr<FRecvMulti> CreateRecvMulti(int32 MaxNumPackets, int32 MaxPacketSize, ERecvMultiFlags Flags) override;
+	virtual bool IsSocketRecvMultiSupported() const override;
+	virtual double TranslatePacketTimestamp(const FPacketTimestamp& Timestamp, ETimestampTranslation Translation) override;
 };
diff --git a/Engine/Source/Runtime/Sockets/Private/Unix/SocketsUnix.cpp b/Engine/Source/Runtime/Sockets/Private/Unix/SocketsUnix.cpp
new file mode 100644
--- /dev/null
+++ b/Engine/Source/Runtime/Sockets/Private/Unix/SocketsUnix.cpp
@@ -0,0 +1,155 @@
+// Copyright 1998-2019 Epic Games, Inc. All Rights Reserved
+
+
+#include "SocketsUnix.h"
+
+
+// @todo: Add timestamp support for normal Recv/RecvFrom (not essential, there is no API for this yet)
+
+
+#if PLATFORM_HAS_BSD_SOCKET_FEATURE_RECVMMSG
+constexpr const int32 TimestampControlMsgSize		= CMSG_SPACE(sizeof(timeval));
+constexpr const int32 TimestampControlReturnSize	= CMSG_LEN(sizeof(timeval));
+
+
+/**
+ * FUnixRecvMulti
+ */
+
+FUnixRecvMulti::FUnixRecvMulti(ISocketSubsystem* SocketSubsystem, int32 InMaxNumPackets, int32 InMaxPacketSize,
+								ERecvMultiFlags InitFlags)
+	: FRecvMulti(SocketSubsystem, InMaxNumPackets, InMaxPacketSize, InitFlags)
+	, Headers(MakeUnique<mmsghdr[]>(MaxNumPackets))
+	, BufferMaps(MakeUnique<iovec[]>(MaxNumPackets))
+	, DataBuffer(MakeUnique<uint8[]>(MaxNumPackets * MaxPacketSize))
+{
+	bool bRetrieveTimestamps = EnumHasAnyFlags(InitFlags, ERecvMultiFlags::RetrieveTimestamps);
+	RawTimestampData = (bRetrieveTimestamps ? MakeUnique<uint8[]>(TimestampControlMsgSize * MaxNumPackets) : nullptr);
+
+	for (int32 i=0; i<MaxNumPackets; i++)
+	{
+		mmsghdr& CurHeader = Headers[i];
+		msghdr& CurInnerHeader = CurHeader.msg_hdr;
+		iovec& CurBufferMap = BufferMaps[i];
+		FRecvMulti::FRecvData& CurPacket = Packets[i];
+
+		CurPacket.Data = &DataBuffer[MaxPacketSize*i];
+		CurPacket.Source = SocketSubsystem->CreateInternetAddr();
+		CurPacket.BytesReadPtr = &CurHeader.msg_len;
+
+
+		FInternetAddrBSD* CurBSDAddr = (FInternetAddrBSD*)CurPacket.Source.Get();
+
+		CurInnerHeader.msg_name = CurBSDAddr->GetRawAddr();
+		CurInnerHeader.msg_namelen = CurBSDAddr->GetStorageSize();
+		CurInnerHeader.msg_iov = &CurBufferMap;
+		CurInnerHeader.msg_iovlen = 1;
+		CurInnerHeader.msg_control = (bRetrieveTimestamps ? &RawTimestampData[i * TimestampControlMsgSize] : nullptr);;
+		CurInnerHeader.msg_controllen = (bRetrieveTimestamps ? TimestampControlMsgSize : 0);
+		CurInnerHeader.msg_flags = 0;
+
+		CurBufferMap.iov_base = (void*)CurPacket.Data;
+		CurBufferMap.iov_len = MaxPacketSize;
+	}
+}
+
+bool FUnixRecvMulti::GetPacketTimestamp(int32 PacketIdx, FPacketTimestamp& OutTimestamp) const
+{
+	bool bSuccess = false;
+
+	if (RawTimestampData.IsValid())
+	{
+		msghdr& CurInnerHeader = Headers[PacketIdx].msg_hdr;
+		const cmsghdr* TimestampMsg = CMSG_FIRSTHDR(&CurInnerHeader);
+
+		if (TimestampMsg != nullptr && TimestampMsg->cmsg_level == SOL_SOCKET &&
+			TimestampMsg->cmsg_type == SCM_TIMESTAMP && TimestampMsg->cmsg_len == TimestampControlReturnSize)
+		{
+			timeval* CurTimestampData = (timeval*)CMSG_DATA(TimestampMsg);
+
+			if (CurTimestampData != nullptr)
+			{
+				OutTimestamp.Timestamp = FTimespan((CurTimestampData->tv_sec * ETimespan::TicksPerSecond) +
+											(CurTimestampData->tv_usec * ETimespan::TicksPerMicrosecond));
+
+				bSuccess = true;
+			}
+		}
+	}
+
+	return bSuccess;
+}
+
+void FUnixRecvMulti::CountBytes(FArchive& Ar) const
+{
+	FRecvMulti::CountBytes(Ar);
+
+	int32 CurSize = sizeof(*this) - sizeof(FRecvMulti);
+
+	Ar.CountBytes(CurSize, CurSize);
+
+	// Headers
+	CurSize = sizeof(mmsghdr) * MaxNumPackets;
+
+	Ar.CountBytes(CurSize, CurSize);
+
+	// RawTimestampData
+	CurSize = (RawTimestampData.IsValid() ? (TimestampControlMsgSize * MaxNumPackets) : 0);
+
+	Ar.CountBytes(CurSize, CurSize);
+
+	// BufferMaps
+	CurSize = sizeof(iovec) * MaxNumPackets;
+
+	Ar.CountBytes(CurSize, CurSize);
+
+	// FInternetAddrBSD
+	CurSize = sizeof(FInternetAddrBSD) * MaxNumPackets;
+
+	Ar.CountBytes(CurSize, CurSize);
+
+	// DataBuffer
+	Ar.CountBytes(MaxNumPackets, MaxNumPackets);
+}
+#endif
+
+
+/**
+ * FSocketUnix
+ */
+
+// NOTE: Does not support TCP at the moment.
+bool FSocketUnix::RecvMulti(FRecvMulti& MultiData, ESocketReceiveFlags::Type Flags)
+{
+	bool bSuccess = false;
+
+#if PLATFORM_HAS_BSD_SOCKET_FEATURE_RECVMMSG
+	FUnixRecvMulti& UnixMultiData = (FUnixRecvMulti&)MultiData;
+	mmsghdr* Headers = UnixMultiData.Headers.Get();
+	const int TranslatedFlags = TranslateFlags(Flags);
+	int NumPacketsRead = recvmmsg(Socket, Headers, UnixMultiData.MaxNumPackets, TranslatedFlags, nullptr);
+
+	bSuccess = NumPacketsRead > 0;
+	UnixMultiData.NumPackets = FMath::Max(0, NumPacketsRead);
+
+	if (bSuccess)
+	{
+		LastActivityTime = FPlatformTime::Seconds();
+	}
+#endif
+
+	return bSuccess;
+}
+
+bool FSocketUnix::SetRetrieveTimestamp(bool bRetrieveTimestamp)
+{
+	bool bSuccess = false;
+
+#if PLATFORM_HAS_BSD_SOCKET_FEATURE_TIMESTAMP
+	int32 Timestamp = (int)bRetrieveTimestamp;
+
+	bSuccess = setsockopt(Socket, SOL_SOCKET, SO_TIMESTAMP, (char*)&Timestamp, sizeof(Timestamp)) == 0;
+#endif
+
+	return bSuccess;
+}
diff --git a/Engine/Source/Runtime/Sockets/Private/Unix/SocketsUnix.h b/Engine/Source/Runtime/Sockets/Private/Unix/SocketsUnix.h
new file mode 100644
--- /dev/null
+++ b/Engine/Source/Runtime/Sockets/Private/Unix/SocketsUnix.h
@@ -0,0 +1,59 @@
+// Copyright 1998-2019 Epic Games, Inc. All Rights Reserved
+
+#pragma once
+
+
+#include "BSDSockets/SocketsBSD.h"
+
+
+#if PLATFORM_HAS_BSD_SOCKET_FEATURE_RECVMMSG
+/**
+ * Implements platform specific data/buffers for RecvMulti in Linux
+ */
+struct FUnixRecvMulti : public FRecvMulti
+{
+	friend class FSocketUnix;
+
+protected:
+	/** Persistently stores preconfigured mmsghdr struct values, for use with recvmmsg */
+	TUniquePtr<mmsghdr[]>	Headers;
+
+	/** Buffer for receiving/storing packet kernel timestamp data, during the call to recvmmsg */
+	TUniquePtr<uint8[]>		RawTimestampData;
+
+
+private:
+	/** Maps sections of DataBuffer for receiving packet data, within Headers */
+	TUniquePtr<iovec[]>		BufferMaps;
+
+	/** The raw data buffer where all packet data is received/stored. */
+	TUniquePtr<uint8[]>		DataBuffer;
+
+
+public:
+	FUnixRecvMulti(ISocketSubsystem* SocketSubsystem, int32 InMaxNumPackets, int32 InMaxPacketSize, ERecvMultiFlags InitFlags);
+
+	virtual bool GetPacketTimestamp(int32 PacketIdx, FPacketTimestamp& OutTimestamp) const override;
+	virtual void CountBytes(FArchive& Ar) const override;
+};
+#endif
+
+
+/**
+ * Unix specific socket implementation - primarily, adds support for recvmmsg
+ */
+class FSocketUnix : public FSocketBSD
+{
+public:
+	/**
+	 * Base constructor
+	 */
+	FSocketUnix(SOCKET InSocket, ESocketType InSocketType, const FString& InSocketDescription, const FName& InSocketProtocol,
+				ISocketSubsystem* InSubsystem)
+		: FSocketBSD(InSocket, InSocketType, InSocketDescription, InSocketProtocol, InSubsystem)
+	{
+	}
+
+	virtual bool RecvMulti(FRecvMulti& MultiData, ESocketReceiveFlags::Type Flags) override;
+	virtual bool SetRetrieveTimestamp(bool bRetrieveTimestamp) override;
+};
diff --git a/Engine/Source/Runtime/Sockets/Public/SocketSubsystem.h b/Engine/Source/Runtime/Sockets/Public/SocketSubsystem.h
--- a/Engine/Source/Runtime/Sockets/Public/SocketSubsystem.h
+++ b/Engine/Source/Runtime/Sockets/Public/SocketSubsystem.h
@@ -9,6 +9,7 @@
 
 class Error;
 class FInternetAddr;
+class FSocket;
 
 SOCKETS_API DECLARE_LOG_CATEGORY_EXTERN(LogSockets, Log, All);
 
@@ -44,7 +45,6 @@ SOCKETS_API DECLARE_LOG_CATEGORY_EXTERN(LogSockets, Log, All);
  */
 class SOCKETS_API ISocketSubsystem
 {
-
 public:
 
 	/**
@@ -83,7 +83,7 @@ public:
 	 *
 	 * @return the new socket or NULL if failed
 	 */
-	virtual class FSocket* CreateSocket(const FName& SocketType, const FString& SocketDescription, bool bForceUDP = false)
+	virtual FSocket* CreateSocket(const FName& SocketType, const FString& SocketDescription, bool bForceUDP = false)
 	{
 		const FName NoProtocolTypeName(NAME_None);
 		return CreateSocket(SocketType, SocketDescription, NoProtocolTypeName);
@@ -99,7 +99,7 @@ public:
 	 * @return the new socket or NULL if failed
 	 */
 	UE_DEPRECATED(4.23, "Use the CreateSocket with the FName parameter for support for multiple protocol types.")
-	virtual class FSocket* CreateSocket(const FName& SocketType, const FString& SocketDescription, ESocketProtocolFamily ProtocolType)
+	virtual FSocket* CreateSocket(const FName& SocketType, const FString& SocketDescription, ESocketProtocolFamily ProtocolType)
 	{
 		return CreateSocket(SocketType, SocketDescription, GetProtocolNameFromFamily(ProtocolType));
 	}
@@ -113,7 +113,7 @@ public:
 	 *
 	 * @return the new socket or NULL if failed
 	 */
-	virtual class FSocket* CreateSocket(const FName& SocketType, const FString& SocketDescription, const FName& ProtocolName) = 0;
+	virtual FSocket* CreateSocket(const FName& SocketType, const FString& SocketDescription, const FName& ProtocolName) = 0;
 
 	/**
 	 * Creates a resolve info cached struct to hold the resolved address
@@ -129,7 +129,7 @@ public:
 	 *
 	 * @param Socket the socket object to destroy
 	 */
-	virtual void DestroySocket(class FSocket* Socket) = 0;
+	virtual void DestroySocket(FSocket* Socket) = 0;
 
 	/**
 	 * Gets the address information of the given hostname and outputs it into an array of resolvable addresses.
@@ -265,6 +265,17 @@ public:
 	 */
 	virtual TSharedRef<FInternetAddr> CreateInternetAddr() = 0;
 
+	/**
+	 * Create a platform specific FRecvMulti representation
+	 *
+	 * @param MaxNumPackets			The maximum number of packet receives supported
+	 * @param MaxPacketSize			The maximum supported packet size
+	 * @param Flags					Flags for specifying how FRecvMulti should be initialized (for e.g. retrieving timestamps)
+	 * @return						Returns the platform specific FRecvMulti instance
+	 */
+	virtual TUniquePtr<FRecvMulti> CreateRecvMulti(int32 MaxNumPackets, int32 MaxPacketSize,
+													ERecvMultiFlags Flags=ERecvMultiFlags::None);
+
 	/**
 	 * @return Whether the machine has a properly configured network device or not
 	 */
@@ -320,7 +331,7 @@ public:
 	 *
 	 * @return The bound port number, or 0 on failure
 	 */
-	int32 BindNextPort(class FSocket* Socket, FInternetAddr& Addr, int32 PortCount, int32 PortIncrement);
+	int32 BindNextPort(FSocket* Socket, FInternetAddr& Addr, int32 PortCount, int32 PortIncrement);
 
 	/**
 	 * Uses the platform specific look up to determine the host address
@@ -367,11 +378,28 @@ public:
 	 */
 	void RemoveHostNameFromCache(const ANSICHAR* HostName);
 
+	/**
+	 * Returns true if FSocket::RecvMulti is supported by this socket subsystem
+	 */
+	virtual bool IsSocketRecvMultiSupported() const;
+
+
 	/**
 	 * Returns true if FSocket::Wait is supported by this socket subsystem.
 	 */
 	virtual bool IsSocketWaitSupported() const = 0;
 
+
+	/**
+	 * Converts a platform packet timestamp, into a local timestamp, or into a time delta etc.
+	 *
+	 * @param Timestamp		The timestamp to translate
+	 * @param Translation	The type of translation to perform on the timestamp (time delta is usually faster than local timestamp)
+	 * @return				Returns the translated timestamp or delta
+	 */
+	virtual double TranslatePacketTimestamp(const FPacketTimestamp& Timestamp,
+											ETimestampTranslation Translation=ETimestampTranslation::LocalTimestamp);
+
 protected:
 
 	/**
diff --git a/Engine/Source/Runtime/Sockets/Public/SocketTypes.h b/Engine/Source/Runtime/Sockets/Public/SocketTypes.h
--- a/Engine/Source/Runtime/Sockets/Public/SocketTypes.h
+++ b/Engine/Source/Runtime/Sockets/Public/SocketTypes.h
@@ -6,6 +6,9 @@
 
 // Forward declarations
 class FInternetAddr;
+class ISocketSubsystem;
+class FSocket;
+
 
 /** Indicates the socket protocol of socket being used, typically for BSD Sockets */
 enum class ESocketProtocolFamily : uint8
@@ -213,3 +216,155 @@ struct FReceivedPacketView
 	/** Error if receiving a packet failed */
 	ESocketErrors				Error;
 };
+
+/**
+ * Stores a platform-specific timestamp for a packet. Can be translated for local use by ISocketSubsystem::TranslatePacketTimestamp.
+ */
+struct FPacketTimestamp
+{
+	/** The internal platform specific timestamp (does NOT correspond to FPlatformTime::Seconds, may use a different clock source). */
+	FTimespan	Timestamp;
+};
+
+
+
+/**
+ * Specifies how a platform specific timestamp (in this case, a packet timestamp) should be translated
+ */
+enum class ETimestampTranslation : uint8
+{
+	/**
+	 * Translates the timestamp into a local timestamp, comparable (with accuracy caveats) to other local timestamps.
+	 *
+	 * Use this to get a value comparable to FPlatformTime::Seconds() - if the platform timestamp uses a different clock,
+	 * this has both a performance (2x internal clock reads i.e. 2x FPlatformTime::Seconds() calls) and accuracy caveat.
+	 */
+	LocalTimestamp,
+
+	/**
+	 * The delta between present platform time vs timestamp - faster than local translation, less accuracy caveats.
+	 *
+	 * This is the most accurate measure of time passed since the packet was recorded at an OS/NIC/Thread level,
+	 * to TranslatePacketTimestamp being called for the packets timestamp - and performs faster (1x FPlatformTime::Seconds() call).
+	 */
+	TimeDelta
+};
+
+/**
+ * Flags for specifying how an FRecvMulti instance should be initialized
+ */
+enum class ERecvMultiFlags : uint32
+{
+	None				= 0x00000000,
+	RetrieveTimestamps	= 0x00000001	// Whether or not to support retrieving timestamps
+};
+
+ENUM_CLASS_FLAGS(ERecvMultiFlags);
+
+
+/**
+ * Stores the persistent state and packet buffers/data, for receiving packets with FSocket::RecvMulti.
+ * To optimize performance, use only once instance of this struct, for the lifetime of the socket.
+ */
+struct SOCKETS_API FRecvMulti : public FNoncopyable, public FVirtualDestructor
+{
+	friend struct FUnixRecvMulti;
+	friend class FSocketUnix;
+
+private:
+	/**
+	 * Receive data for each individual packet
+	 */
+	struct FRecvData
+	{
+		/** The source address for the packet */
+		TSharedPtr<FInternetAddr>	Source;
+
+		/** Pointer to the packet data */
+		const uint8*				Data;
+
+		/** Internal pointer specifying the number of bytes read */
+		const uint32*				BytesReadPtr;
+
+
+		FRecvData()
+			: Source()
+			, Data(nullptr)
+			, BytesReadPtr(nullptr)
+		{
+		}
+	};
+
+
+private:
+	/** The current list of received packets */
+	TUniquePtr<FRecvData[]>			Packets;
+
+	/** The number of packets received */
+	int32							NumPackets;
+
+public:
+	/** The maximum number of packets this FRecvMulti instance can support */
+	const int32						MaxNumPackets;
+
+	/** The maximum packet size this FRecvMulti instance can support */
+	const int32						MaxPacketSize;
+
+
+private:
+	/**
+	 * Initialize an FRecvMulti instance, supporting the specified maximum packet count/sizes
+	 *
+	 * @param SocketSubsystem		The socket subsystem initializing this FRecvMulti instance
+	 * @param InMaxNumPackets		The maximum number of packet receives supported
+	 * @param InMaxPacketSize		The maximum supported packet size
+	 * @param bRetrieveTimestamps	Whether or not to support retrieving timestamps
+	 */
+	FRecvMulti(ISocketSubsystem* SocketSubsystem, int32 InMaxNumPackets, int32 InMaxPacketSize,
+				ERecvMultiFlags InitFlags=ERecvMultiFlags::None);
+
+
+public:
+	/**
+	 * Retrieves the information for the specified packet
+	 *
+	 * @param PacketIdx		The RecvMulti index for the packet to be retrieved
+	 * @param OutPacket		Outputs a view to the received packet
+	 */
+	void GetPacket(int32 PacketIdx, FReceivedPacketView& OutPacket)
+	{
+		check(PacketIdx >= 0);
+		check(PacketIdx < NumPackets);
+
+		FRecvData& CurData = Packets.Get()[PacketIdx];
+
+		OutPacket.Data = MakeArrayView(CurData.Data, *CurData.BytesReadPtr);
+		OutPacket.Address = CurData.Source;
+		OutPacket.Error = ESocketErrors::SE_NO_ERROR;
+	}
+
+	/**
+	 * Returns the platform specific timestamp for when the specified packet was received by the operating system
+	 *
+	 * @param PacketIdx		The index into Packets, of the packet to be checked
+	 * @param OutTimestamp	The timestamp for the specified packet
+	 * @return				Whether or not the timestamp was retrieved successfully
+	 */
+	virtual bool GetPacketTimestamp(int32 PacketIdx, FPacketTimestamp& OutTimestamp) const = 0;
+
+	/**
+	 * Retrieves the current number of received packets
+	 */
+	int32 GetNumPackets() const
+	{
+		return NumPackets;
+	}
+
+
+	/**
+	 * Calculates the total memory consumption of this FRecvMulti instance, including platform-specific data
+	 *
+	 * @param Ar	The archive being used to count the memory consumption
+	 */
+	virtual void CountBytes(FArchive& Ar) const;
+};
diff --git a/Engine/Source/Runtime/Sockets/Public/Sockets.h b/Engine/Source/Runtime/Sockets/Public/Sockets.h
--- a/Engine/Source/Runtime/Sockets/Public/Sockets.h
+++ b/Engine/Source/Runtime/Sockets/Public/Sockets.h
@@ -216,6 +216,17 @@ public:
 	 */
 	virtual bool Recv(uint8* Data, int32 BufferSize, int32& BytesRead, ESocketReceiveFlags::Type Flags = ESocketReceiveFlags::None);
 
+	/**
+	 * Reads multiple packets from the socket at once, gathering the source address and other optional platform specific data.
+	 * Use ISocketSubsystem::IsSocketRecvMultiSupported to check if the current socket platform supports this.
+	 * NOTE: For optimal performance, one FRecvMulti instance should be used, for the lifetime of the socket.
+	 *
+	 * @param MultiData		The FRecvMulti instance that receives packet data and holds platform specific buffers for receiving data.
+	 * @param Flags			The receive flags.
+	 * @return				Whether or not data was successfully received
+	 */
+	virtual bool RecvMulti(FRecvMulti& MultiData, ESocketReceiveFlags::Type Flags=ESocketReceiveFlags::None);
+
 	/**
 	 * Blocks until the specified condition is met.
 	 *
@@ -385,6 +396,14 @@ public:
 	 */
 	virtual bool SetReceiveBufferSize(int32 Size, int32& NewSize) = 0;
 
+	/**
+	 * Sets whether to retrieve the system-level receive timestamp, for sockets
+	 *
+	 * @param bRetrieveTimestamp	Whether to retrieve the timestamp upon receive
+	 * @return						True if the call succeeded, false otherwise.
+	 */
+	virtual bool SetRetrieveTimestamp(bool bRetrieveTimestamp=true);
+
 	/**
 	 * Reads the port this socket is bound to.
 	 *
