Fixed sizeof(this) which should be sizeof(*this).

diff --git a/Engine/Source/Runtime/Sockets/Private/Sockets.cpp b/Engine/Source/Runtime/Sockets/Private/Sockets.cpp
--- a/Engine/Source/Runtime/Sockets/Private/Sockets.cpp
+++ b/Engine/Source/Runtime/Sockets/Private/Sockets.cpp
@@ -19,7 +19,7 @@ FRecvMulti::FRecvMulti(ISocketSubsystem* SocketSubsystem, int32 InMaxNumPackets,
 
 void FRecvMulti::CountBytes(FArchive& Ar) const
 {
-	Ar.CountBytes(sizeof(this), sizeof(this));
+	Ar.CountBytes(sizeof(*this), sizeof(*this));
 
 	// Packets
 	Ar.CountBytes(MaxNumPackets * sizeof(FRecvData), MaxNumPackets * sizeof(FRecvData));
