Fix the iOS inheritance issue with CreateInternetAddr and the type oddities with IPAddressBSDIOS when setting multicast

#rnx

diff --git a/Engine/Source/Runtime/Sockets/Private/IOS/IPAddressBSDIOS.h b/Engine/Source/Runtime/Sockets/Private/IOS/IPAddressBSDIOS.h
--- a/Engine/Source/Runtime/Sockets/Private/IOS/IPAddressBSDIOS.h
+++ b/Engine/Source/Runtime/Sockets/Private/IOS/IPAddressBSDIOS.h
@@ -3,6 +3,7 @@
 #pragma once
 
 #include "CoreMinimal.h"
+#include "SocketSubsystemIOS.h"
 #include "BSDSockets/SocketSubsystemBSDPrivate.h"
 #include "BSDSockets/IPAddressBSD.h"
 
@@ -18,10 +19,10 @@ public:
 	/** Sets the address to broadcast */
 	virtual void SetIPv6BroadcastAddress() override
 	{
-		FSocketSubsystemBSD* SocketSubsystem = static_cast<FSocketSubsystemBSD*>(ISocketSubsystem::Get(PLATFORM_SOCKETSUBSYSTEM));
-		if (SocketSubsystem)
+		FSocketSubsystemIOS* SocketSubsystemIOS = static_cast<FSocketSubsystemIOS*>(SocketSubsystem);
+		if (SocketSubsystemIOS)
 		{
-			TSharedPtr<FInternetAddrBSD> MultiCastAddr = StaticCastSharedPtr<FInternetAddrBSD>(SocketSubsystem->GetAddressFromString(TEXT("ff02::01")));
+			TSharedPtr<FInternetAddrBSDIOS> MultiCastAddr = StaticCastSharedPtr<FInternetAddrBSDIOS>(SocketSubsystemIOS->GetAddressFromString(TEXT("ff02::01")));
 			if (!MultiCastAddr.IsValid())
 			{
 				UE_LOG(LogSockets, Warning, TEXT("Could not resolve the broadcast address for iOS, this address will just be blank"));
@@ -34,7 +35,7 @@ public:
 				// Do a query to get the scope id of the address.
 				bool bUnusedBool;
 				TSharedRef<FInternetAddrBSDIOS> ScopeAddr = 
-					StaticCastSharedRef<FInternetAddrBSDIOS>(SocketSubsystem->GetLocalHostAddr(*GLog, bUnusedBool));
+					StaticCastSharedRef<FInternetAddrBSDIOS>(SocketSubsystemIOS->GetLocalHostAddr(*GLog, bUnusedBool));
 				SetScopeId(ScopeAddr->GetScopeId());
 			}
 		}
