Adding TCP NoDelay functionality based off PR#3752, updated to latest codebase.

diff --git a/Engine/Plugins/Online/OnlineSubsystemSteam/Source/Private/SocketsSteam.h b/Engine/Plugins/Online/OnlineSubsystemSteam/Source/Private/SocketsSteam.h
--- a/Engine/Plugins/Online/OnlineSubsystemSteam/Source/Private/SocketsSteam.h
+++ b/Engine/Plugins/Online/OnlineSubsystemSteam/Source/Private/SocketsSteam.h
@@ -215,6 +215,11 @@ public:
 	 */
 	virtual bool SetBroadcast(bool bAllowBroadcast = true) override;
 
+	virtual bool SetNoDelay(bool bIsNoDelay = true) override
+	{
+		return true;
+	}
+
 	virtual bool JoinMulticastGroup(const FInternetAddr& GroupAddress) override;
 
 	virtual bool LeaveMulticastGroup(const FInternetAddr& GroupAddress) override;
diff --git a/Engine/Source/Runtime/Core/Public/HAL/Platform.h b/Engine/Source/Runtime/Core/Public/HAL/Platform.h
--- a/Engine/Source/Runtime/Core/Public/HAL/Platform.h
+++ b/Engine/Source/Runtime/Core/Public/HAL/Platform.h
@@ -284,6 +284,9 @@
 #ifndef PLATFORM_HAS_BSD_SOCKET_FEATURE_MSG_DONTWAIT
 	#define PLATFORM_HAS_BSD_SOCKET_FEATURE_MSG_DONTWAIT	0
 #endif
+#ifndef PLATFORM_HAS_BSD_SOCKET_FEATURE_NODELAY
+	#define PLATFORM_HAS_BSD_SOCKET_FEATURE_NODELAY	1
+#endif
 #ifndef PLATFORM_HAS_NO_EPROCLIM
 	#define PLATFORM_HAS_NO_EPROCLIM			0
 #endif
diff --git a/Engine/Source/Runtime/Sockets/Private/BSDIPv6Sockets/SocketsBSDIPv6.cpp b/Engine/Source/Runtime/Sockets/Private/BSDIPv6Sockets/SocketsBSDIPv6.cpp
--- a/Engine/Source/Runtime/Sockets/Private/BSDIPv6Sockets/SocketsBSDIPv6.cpp
+++ b/Engine/Source/Runtime/Sockets/Private/BSDIPv6Sockets/SocketsBSDIPv6.cpp
@@ -387,6 +387,19 @@ bool FSocketBSDIPv6::SetBroadcast(bool bAllowBroadcast)
 	return setsockopt(Socket,SOL_SOCKET,SO_BROADCAST,(char*)&Param,sizeof(Param)) == 0;
 }
 
+bool FSocketBSDIPv6::SetNoDelay(bool bIsNoDelay)
+{
+	if (GetSocketType() == SOCKTYPE_Streaming)
+	{
+#if PLATFORM_HAS_BSD_SOCKET_FEATURE_NODELAY
+		int Param = bIsNoDelay ? 1 : 0;
+		return setsockopt(Socket, IPPROTO_TCP, TCP_NODELAY, (char*)&Param, sizeof(Param)) == 0;
+#endif
+	}
+
+	return true;
+}
+
 
 bool FSocketBSDIPv6::JoinMulticastGroup(const FInternetAddr& GroupAddress)
 {
diff --git a/Engine/Source/Runtime/Sockets/Private/BSDIPv6Sockets/SocketsBSDIPv6.h b/Engine/Source/Runtime/Sockets/Private/BSDIPv6Sockets/SocketsBSDIPv6.h
--- a/Engine/Source/Runtime/Sockets/Private/BSDIPv6Sockets/SocketsBSDIPv6.h
+++ b/Engine/Source/Runtime/Sockets/Private/BSDIPv6Sockets/SocketsBSDIPv6.h
@@ -125,6 +125,8 @@ public:
 
 	virtual bool SetBroadcast(bool bAllowBroadcast = true) override;
 
+	virtual bool SetNoDelay(bool bIsNoDelay = true) override;
+
 	virtual bool JoinMulticastGroup(const FInternetAddr& GroupAddress) override;
 
 	virtual bool LeaveMulticastGroup(const FInternetAddr& GroupAddress) override;
diff --git a/Engine/Source/Runtime/Sockets/Private/BSDSockets/SocketsBSD.cpp b/Engine/Source/Runtime/Sockets/Private/BSDSockets/SocketsBSD.cpp
--- a/Engine/Source/Runtime/Sockets/Private/BSDSockets/SocketsBSD.cpp
+++ b/Engine/Source/Runtime/Sockets/Private/BSDSockets/SocketsBSD.cpp
@@ -63,7 +63,6 @@ bool FSocketBSD::Close(void)
 	return false;
 }
 
-
 bool FSocketBSD::Bind(const FInternetAddr& Addr)
 {
 	if (Addr.GetProtocolType() != GetProtocol())
@@ -77,7 +76,6 @@ bool FSocketBSD::Bind(const FInternetAddr& Addr)
 	return bind(Socket, (const sockaddr*)&(BSDAddr.Addr), BSDAddr.GetStorageSize()) == 0;
 }
 
-
 bool FSocketBSD::Connect(const FInternetAddr& Addr)
 {
 	if (Addr.GetProtocolType() != GetProtocol())
@@ -103,7 +101,6 @@ bool FSocketBSD::Listen(int32 MaxBacklog)
 	return listen(Socket, MaxBacklog) == 0;
 }
 
-
 bool FSocketBSD::WaitForPendingConnection(bool& bHasPendingConnection, const FTimespan& WaitTime)
 {
 	bool bHasSucceeded = false;
@@ -123,7 +120,6 @@ bool FSocketBSD::WaitForPendingConnection(bool& bHasPendingConnection, const FTi
 	return bHasSucceeded;
 }
 
-
 bool FSocketBSD::HasPendingData(uint32& PendingDataSize)
 {
 	PendingDataSize = 0;
@@ -143,7 +139,6 @@ bool FSocketBSD::HasPendingData(uint32& PendingDataSize)
 	return false;
 }
 
-
 FSocket* FSocketBSD::Accept(const FString& InSocketDescription)
 {
 	SOCKET NewSocket = accept(Socket, NULL, NULL);
@@ -159,7 +154,6 @@ FSocket* FSocketBSD::Accept(const FString& InSocketDescription)
 	return NULL;
 }
 
-
 FSocket* FSocketBSD::Accept(FInternetAddr& OutAddr, const FString& InSocketDescription)
 {
 	FInternetAddrBSD& BSDAddr = static_cast<FInternetAddrBSD&>(OutAddr);
@@ -177,7 +171,6 @@ FSocket* FSocketBSD::Accept(FInternetAddr& OutAddr, const FString& InSocketDescr
 	return NULL;
 }
 
-
 bool FSocketBSD::SendTo(const uint8* Data, int32 Count, int32& BytesSent, const FInternetAddr& Destination)
 {
 	// TODO: Consider converting IPv4 addresses to v6 when needed
@@ -200,7 +193,6 @@ bool FSocketBSD::SendTo(const uint8* Data, int32 Count, int32& BytesSent, const
 	return Result;
 }
 
-
 bool FSocketBSD::Send(const uint8* Data, int32 Count, int32& BytesSent)
 {
 	BytesSent = send(Socket,(const char*)Data,Count,0);
@@ -215,7 +207,6 @@ bool FSocketBSD::Send(const uint8* Data, int32 Count, int32& BytesSent)
 	return Result;
 }
 
-
 bool FSocketBSD::RecvFrom(uint8* Data, int32 BufferSize, int32& BytesRead, FInternetAddr& Source, ESocketReceiveFlags::Type Flags)
 {
 	bool bSuccess = false;
@@ -249,7 +240,6 @@ bool FSocketBSD::RecvFrom(uint8* Data, int32 BufferSize, int32& BytesRead, FInte
 	return bSuccess;
 }
 
-
 bool FSocketBSD::Recv(uint8* Data, int32 BufferSize, int32& BytesRead, ESocketReceiveFlags::Type Flags)
 {
 	bool bSuccess = false;
@@ -279,7 +269,6 @@ bool FSocketBSD::Recv(uint8* Data, int32 BufferSize, int32& BytesRead, ESocketRe
 	return bSuccess;
 }
 
-
 bool FSocketBSD::Wait(ESocketWaitConditions::Type Condition, FTimespan WaitTime)
 {
 	if ((Condition == ESocketWaitConditions::WaitForRead) || (Condition == ESocketWaitConditions::WaitForReadOrWrite))
@@ -301,7 +290,6 @@ bool FSocketBSD::Wait(ESocketWaitConditions::Type Condition, FTimespan WaitTime)
 	return false;
 }
 
-
 ESocketConnectionState FSocketBSD::GetConnectionState(void)
 {
 	ESocketConnectionState CurrentState = SCS_ConnectionError;
@@ -335,7 +323,6 @@ ESocketConnectionState FSocketBSD::GetConnectionState(void)
 	return CurrentState;
 }
 
-
 void FSocketBSD::GetAddress(FInternetAddr& OutAddr)
 {
 	FInternetAddrBSD& BSDAddr = static_cast<FInternetAddrBSD&>(OutAddr);
@@ -351,7 +338,6 @@ void FSocketBSD::GetAddress(FInternetAddr& OutAddr)
 	}
 }
 
-
 bool FSocketBSD::GetPeerAddress(FInternetAddr& OutAddr)
 {
 	FInternetAddrBSD& BSDAddr = static_cast<FInternetAddrBSD&>(OutAddr);
@@ -390,6 +376,19 @@ bool FSocketBSD::SetNonBlocking(bool bIsNonBlocking)
 #endif
 }
 
+bool FSocketBSD::SetNoDelay(bool bIsNoDelay)
+{
+	// Only do this if we are TCP.
+	if (GetSocketType() == SOCKTYPE_Streaming)
+	{
+#if PLATFORM_HAS_BSD_SOCKET_FEATURE_NODELAY
+		int Param = bIsNoDelay ? 1 : 0;
+		return setsockopt(Socket, IPPROTO_TCP, TCP_NODELAY, (char*)&Param, sizeof(Param)) == 0;
+#endif
+	}
+
+	return true;
+}
 
 bool FSocketBSD::SetBroadcast(bool bAllowBroadcast)
 {
@@ -397,7 +396,6 @@ bool FSocketBSD::SetBroadcast(bool bAllowBroadcast)
 	return setsockopt(Socket,SOL_SOCKET,SO_BROADCAST,(char*)&Param,sizeof(Param)) == 0;
 }
 
-
 bool FSocketBSD::JoinMulticastGroup(const FInternetAddr& GroupAddress)
 {
 	const FInternetAddrBSD& BSDAddr = static_cast<const FInternetAddrBSD&>(GroupAddress);
@@ -418,7 +416,6 @@ bool FSocketBSD::JoinMulticastGroup(const FInternetAddr& GroupAddress)
 	return (setsockopt(Socket, IPPROTO_IP, IP_ADD_MEMBERSHIP, (char*)&imr, sizeof(imr)) == 0);
 }
 
-
 bool FSocketBSD::JoinMulticastGroup(const FInternetAddr& GroupAddress, const FInternetAddr& InterfaceAddress)
 {
 	const FInternetAddrBSD& BSDAddr = static_cast<const FInternetAddrBSD&>(GroupAddress);
@@ -440,7 +437,6 @@ bool FSocketBSD::JoinMulticastGroup(const FInternetAddr& GroupAddress, const FIn
 	return (setsockopt(Socket, IPPROTO_IP, IP_ADD_MEMBERSHIP, (char*)&imr, sizeof(imr)) == 0);
 }
 
-
 bool FSocketBSD::LeaveMulticastGroup(const FInternetAddr& GroupAddress)
 {
 	const FInternetAddrBSD& BSDAddr = static_cast<const FInternetAddrBSD&>(GroupAddress);
@@ -461,7 +457,6 @@ bool FSocketBSD::LeaveMulticastGroup(const FInternetAddr& GroupAddress)
 	return (setsockopt(Socket, IPPROTO_IP, IP_DROP_MEMBERSHIP, (char*)&imr, sizeof(imr)) == 0);
 }
 
-
 bool FSocketBSD::LeaveMulticastGroup(const FInternetAddr& GroupAddress, const FInternetAddr& InterfaceAddress)
 {
 	const FInternetAddrBSD& BSDAddr = static_cast<const FInternetAddrBSD&>(GroupAddress);
@@ -483,7 +478,6 @@ bool FSocketBSD::LeaveMulticastGroup(const FInternetAddr& GroupAddress, const FI
 	return (setsockopt(Socket, IPPROTO_IP, IP_DROP_MEMBERSHIP, (char*)&imr, sizeof(imr)) == 0);
 }
 
-
 bool FSocketBSD::SetMulticastLoopback(bool bLoopback)
 {
 #if PLATFORM_HAS_BSD_IPV6_SOCKETS
@@ -497,7 +491,6 @@ bool FSocketBSD::SetMulticastLoopback(bool bLoopback)
 	return (setsockopt(Socket, IPPROTO_IP, IP_MULTICAST_LOOP, (char*)&bLoopback, sizeof(bLoopback)) == 0);
 }
 
-
 bool FSocketBSD::SetMulticastTtl(uint8 TimeToLive)
 {
 #if PLATFORM_HAS_BSD_IPV6_SOCKETS
@@ -511,7 +504,6 @@ bool FSocketBSD::SetMulticastTtl(uint8 TimeToLive)
 	return (setsockopt(Socket, IPPROTO_IP, IP_MULTICAST_TTL, (char*)&TimeToLive, sizeof(TimeToLive)) == 0);
 }
 
-
 bool FSocketBSD::SetMulticastInterface(const FInternetAddr& InterfaceAddress)
 {
 	const FInternetAddrBSD& BSDIFAddr = static_cast<const FInternetAddrBSD&>(InterfaceAddress);
@@ -528,7 +520,6 @@ bool FSocketBSD::SetMulticastInterface(const FInternetAddr& InterfaceAddress)
 	return (setsockopt(Socket, IPPROTO_IP, IP_MULTICAST_IF, (char*)&InterfaceAddr, sizeof(InterfaceAddr)) == 0);
 }
 
-
 bool FSocketBSD::SetReuseAddr(bool bAllowReuse)
 {
 	int Param = bAllowReuse ? 1 : 0;
@@ -542,7 +533,6 @@ bool FSocketBSD::SetReuseAddr(bool bAllowReuse)
 	return ReuseAddrResult == 0;
 }
 
-
 bool FSocketBSD::SetLinger(bool bShouldLinger,int32 Timeout)
 {
 	linger ling;
@@ -553,14 +543,12 @@ bool FSocketBSD::SetLinger(bool bShouldLinger,int32 Timeout)
 	return setsockopt(Socket,SOL_SOCKET,SO_LINGER,(char*)&ling,sizeof(ling)) == 0;
 }
 
-
 bool FSocketBSD::SetRecvErr(bool bUseErrorQueue)
 {
 	// Not supported, but return true to avoid spurious log messages
 	return true;
 }
 
-
 bool FSocketBSD::SetSendBufferSize(int32 Size,int32& NewSize)
 {
 	SOCKLEN SizeSize = sizeof(int32);
@@ -572,7 +560,6 @@ bool FSocketBSD::SetSendBufferSize(int32 Size,int32& NewSize)
 	return bOk;
 }
 
-
 bool FSocketBSD::SetReceiveBufferSize(int32 Size,int32& NewSize)
 {
 	SOCKLEN SizeSize = sizeof(int32);
@@ -584,7 +571,6 @@ bool FSocketBSD::SetReceiveBufferSize(int32 Size,int32& NewSize)
 	return bOk;
 }
 
-
 int32 FSocketBSD::GetPortNo(void)
 {
 	sockaddr_storage Addr;
diff --git a/Engine/Source/Runtime/Sockets/Private/BSDSockets/SocketsBSD.h b/Engine/Source/Runtime/Sockets/Private/BSDSockets/SocketsBSD.h
--- a/Engine/Source/Runtime/Sockets/Private/BSDSockets/SocketsBSD.h
+++ b/Engine/Source/Runtime/Sockets/Private/BSDSockets/SocketsBSD.h
@@ -117,6 +117,7 @@ public:
 	virtual void GetAddress(FInternetAddr& OutAddr) override;
 	virtual bool GetPeerAddress(FInternetAddr& OutAddr) override;
 	virtual bool SetNonBlocking(bool bIsNonBlocking = true) override;
+	virtual bool SetNoDelay(bool bIsNoDelay = true) override;
 	virtual bool SetBroadcast(bool bAllowBroadcast = true) override;
 	virtual bool JoinMulticastGroup(const FInternetAddr& GroupAddress) override;
 	virtual bool JoinMulticastGroup(const FInternetAddr& GroupAddress, const FInternetAddr& InterfaceAddress) override;
diff --git a/Engine/Source/Runtime/Sockets/Public/Sockets.h b/Engine/Source/Runtime/Sockets/Public/Sockets.h
--- a/Engine/Source/Runtime/Sockets/Public/Sockets.h
+++ b/Engine/Source/Runtime/Sockets/Public/Sockets.h
@@ -263,6 +263,14 @@ public:
 	 */
 	virtual bool SetBroadcast(bool bAllowBroadcast = true) = 0;
 
+	/**
+	 * Sets this socket into TCP_NODELAY mode (TCP only).
+	 *
+	 * @param bIsNoDelay Whether to enable no delay mode.
+	 * @return true if successful, false otherwise.
+	 */
+	virtual bool SetNoDelay(bool bIsNoDelay = true) = 0;
+
 	/**
 	 * Joins this socket to the specified multicast group.
 	 *
