Support for Async GAI queries.

diff --git a/Engine/Source/Runtime/Sockets/Private/SocketSubsystem.cpp b/Engine/Source/Runtime/Sockets/Private/SocketSubsystem.cpp
--- a/Engine/Source/Runtime/Sockets/Private/SocketSubsystem.cpp
+++ b/Engine/Source/Runtime/Sockets/Private/SocketSubsystem.cpp
@@ -5,6 +5,7 @@
 #include "Misc/ScopeLock.h"
 #include "SocketSubsystemModule.h"
 #include "Modules/ModuleManager.h"
+#include "Async/AsyncWork.h"
 #include "SocketTypes.h"
 #include "IPAddress.h"
 #include "Sockets.h"
@@ -14,16 +15,17 @@ DEFINE_LOG_CATEGORY(LogSockets);
 
 IMPLEMENT_MODULE( FSocketSubsystemModule, Sockets );
 
-/** Each platform will implement these functions to construct/destroy socket implementations */
-extern FName CreateSocketSubsystem( FSocketSubsystemModule& SocketSubsystemModule );
-extern void DestroySocketSubsystem( FSocketSubsystemModule& SocketSubsystemModule );
-
+/** Values for the Protocol Typenames */
 namespace FNetworkProtocolTypes
 {
 	const FName IPv4(TEXT("IPv4"));
 	const FName IPv6(TEXT("IPv6"));
 }
 
+/** Each platform will implement these functions to construct/destroy socket implementations */
+extern FName CreateSocketSubsystem(FSocketSubsystemModule& SocketSubsystemModule);
+extern void DestroySocketSubsystem(FSocketSubsystemModule& SocketSubsystemModule);
+
 /** Helper function to turn the friendly subsystem name into the module name */
 static inline FName GetSocketModuleName(const FString& SubsystemName)
 {
@@ -258,6 +260,53 @@ int32 ISocketSubsystem::BindNextPort(FSocket* Socket, FInternetAddr& Addr, int32
 	return 0;
 }
 
+/** Async task support for GetAddressInfo */
+class FGetAddressInfoTask : public FNonAbandonableTask
+{
+public:
+	friend class FAutoDeleteAsyncTask<FGetAddressInfoTask>;
+
+	FGetAddressInfoTask(class ISocketSubsystem* InSocketSubsystem, const FString& InQueryHost,
+		const FString& InQueryService, EAddressInfoFlags InQueryFlags, const FName& InQueryProtocol,
+		ESocketType InQuerySocketType, FAsyncGetAddressInfoCallback InCallbackFunction) :
+
+		SocketSubsystem(InSocketSubsystem),
+		QueryHost(InQueryHost),
+		QueryService(InQueryService),
+		QueryFlags(InQueryFlags),
+		QueryProtocol(InQueryProtocol),
+		QuerySocketType(InQuerySocketType),
+		CallbackFunction(InCallbackFunction)
+	{
+
+	}
+
+	void DoWork()
+	{
+		CallbackFunction(SocketSubsystem->GetAddressInfo(*QueryHost, *QueryService, QueryFlags, QueryProtocol, QuerySocketType));
+	}
+
+	FORCEINLINE TStatId GetStatId() const
+	{
+		RETURN_QUICK_DECLARE_CYCLE_STAT(FGetAddressInfoTask, STATGROUP_ThreadPoolAsyncTasks);
+	}
+
+private:
+	ISocketSubsystem* SocketSubsystem;
+	const FString QueryHost;
+	const FString QueryService;
+	EAddressInfoFlags QueryFlags;
+	const FName QueryProtocol;
+	ESocketType QuerySocketType;
+	FAsyncGetAddressInfoCallback CallbackFunction;
+};
+
+void ISocketSubsystem::GetAddressInfoAsync(FAsyncGetAddressInfoCallback Callback, const TCHAR* HostName, const TCHAR* ServiceName,
+	EAddressInfoFlags QueryFlags, const FName ProtocolTypeName, ESocketType SocketType)
+{
+	(new FAutoDeleteAsyncTask<FGetAddressInfoTask>(this, HostName, ServiceName, QueryFlags, ProtocolTypeName, SocketType, Callback))->StartBackgroundTask();
+}
+
 TSharedRef<FInternetAddr> ISocketSubsystem::GetLocalBindAddr(FOutputDevice& Out)
 {
 	bool bCanBindAll;
@@ -490,26 +539,3 @@ const TCHAR* ISocketSubsystem::GetSocketError(ESocketErrors Code)
 		default: return TEXT("Unknown Error");
 	};
 }
-
-
-
-
-
-
-
-
-
-
-
-
-
-
-
-
-
-
-
-
-
-
-
diff --git a/Engine/Source/Runtime/Sockets/Public/AddressInfoTypes.h b/Engine/Source/Runtime/Sockets/Public/AddressInfoTypes.h
--- a/Engine/Source/Runtime/Sockets/Public/AddressInfoTypes.h
+++ b/Engine/Source/Runtime/Sockets/Public/AddressInfoTypes.h
@@ -116,3 +116,4 @@ struct FAddressInfoResult
 	TArray<FAddressInfoResultData> Results;
 };
 
+typedef TFunction<void(FAddressInfoResult)> FAsyncGetAddressInfoCallback;
diff --git a/Engine/Source/Runtime/Sockets/Public/SocketSubsystem.h b/Engine/Source/Runtime/Sockets/Public/SocketSubsystem.h
--- a/Engine/Source/Runtime/Sockets/Public/SocketSubsystem.h
+++ b/Engine/Source/Runtime/Sockets/Public/SocketSubsystem.h
@@ -180,11 +180,39 @@ public:
 		const FName ProtocolTypeName = NAME_None,
 		ESocketType SocketType = ESocketType::SOCKTYPE_Unknown) = 0;
 
+	/**
+	 * Async variant of GetAddressInfo that fetches the data from the above function in an 
+	 * asynchronous task executed on an available background thread.
+	 * 
+	 * On Completion, this fires a callback function that will be executed on the same thread as
+	 * the task's execution. The caller is responsible for either dispatching the result to the thread of their choosing
+	 * or to allow the result callback execute on the task's thread.
+	 *
+	 * This function allows for specifying FNames for the protocol type, allowing for support of other
+	 * platform protocols
+	 *
+	 * @param Callback the callback function to fire when this query completes. Contains the FAddressInfoResult structure.
+	 * @param HostName string version of the queryable hostname or ip address
+	 * @param ServiceName string version of a service name ("http") or a port number ("80")
+	 * @param QueryFlags What flags are used in making the getaddrinfo call. Several flags can be used at once by
+	 *                   bitwise OR-ing the flags together.
+	 *                   Platforms are required to translate this value into a the correct flag representation.
+	 * @param ProtocolTypeName Used to limit results from the call. Specifying None will search all valid protocols.
+	 *					   Callers will find they rarely have to specify this flag.
+	 * @param SocketType What socket type should the results be formatted for. This typically does not change any
+	 *                   formatting results and can be safely left to the default value.
+	 *
+	 */
+	virtual void GetAddressInfoAsync(FAsyncGetAddressInfoCallback Callback, const TCHAR* HostName,
+		const TCHAR* ServiceName = nullptr, EAddressInfoFlags QueryFlags = EAddressInfoFlags::Default,
+		const FName ProtocolTypeName = NAME_None,
+		ESocketType SocketType = ESocketType::SOCKTYPE_Unknown);
+
 	/**
 	 * Serializes a string that only contains an address.
 	 * 
 	 * This is a what you see is what you get, there is no DNS resolution of the input string,
-	 * so only use this if you know you already have a valid address.
+	 * so only use this if you know you already have a valid address and will not need to convert.
 	 * Otherwise, feed the address to GetAddressInfo for guaranteed results.
 	 *
 	 * @param InAddress the address to serialize
