Basic GAI Cheats functionality for QA.

diff --git a/Engine/Source/Runtime/Sockets/Private/SocketSubsystem.cpp b/Engine/Source/Runtime/Sockets/Private/SocketSubsystem.cpp
--- a/Engine/Source/Runtime/Sockets/Private/SocketSubsystem.cpp
+++ b/Engine/Source/Runtime/Sockets/Private/SocketSubsystem.cpp
@@ -3,6 +3,7 @@
 #include "SocketSubsystem.h"
 #include "Misc/CommandLine.h"
 #include "Misc/ScopeLock.h"
+#include "Misc/CoreMisc.h"
 #include "SocketSubsystemModule.h"
 #include "Modules/ModuleManager.h"
 #include "Async/AsyncWork.h"
@@ -539,3 +540,122 @@ const TCHAR* ISocketSubsystem::GetSocketError(ESocketErrors Code)
 		default: return TEXT("Unknown Error");
 	};
 }
+
+//////////////////////////////////
+// SocketSubsystem Testing
+/////////////////////////////////
+#if WITH_DEV_AUTOMATION_TESTS && !UE_BUILD_SHIPPING
+
+static void DebugPrintGaiResults(ISocketSubsystem* SocketSub, const FAddressInfoResult& GAIResult)
+{
+	UE_LOG(LogSockets, Log, TEXT("Got %d GAI Results for hostname %s. Error Code: %s [%d]"), GAIResult.Results.Num(), *GAIResult.QueryHostName,
+		SocketSub->GetSocketError(GAIResult.ReturnCode), GAIResult.ReturnCode);
+
+	for (int i = 0; i < GAIResult.Results.Num(); ++i)
+	{
+		UE_LOG(LogSockets, Log, TEXT("Result #%d Address: %s Type: %s"), i, *GAIResult.Results[i].Address->ToString(false),
+			*GAIResult.Results[i].Address->GetProtocolType().ToString());
+	}
+}
+
+static void RunGAIQuery(const FString& HostStr)
+{
+	if (HostStr.IsEmpty())
+	{
+		UE_LOG(LogSockets, Warning, TEXT("SOCKETSUB GAI requires an input string to test with."));
+		return;
+	}
+
+	ISocketSubsystem* SocketSub = ISocketSubsystem::Get();
+	if (SocketSub)
+	{
+		FAddressInfoResult GAIResult = SocketSub->GetAddressInfo(*HostStr, nullptr, EAddressInfoFlags::AllResultsWithMapping | EAddressInfoFlags::OnlyUsableAddresses, NAME_None);
+		if (GAIResult.Results.Num() > 0)
+		{
+			DebugPrintGaiResults(SocketSub, GAIResult);
+		}
+		else
+		{
+			UE_LOG(LogSockets, Warning, TEXT("Did not get results!"));
+		}
+		return;
+	}
+
+	UE_LOG(LogSockets, Warning, TEXT("Failed to get socket subsystem!"));
+}
+
+static void RunAsyncGAIQuery(const FString& HostStr)
+{
+	if (HostStr.IsEmpty())
+	{
+		UE_LOG(LogSockets, Warning, TEXT("SOCKETSUB ASYNCGAI requires an input string to test with."));
+		return;
+	}
+
+	ISocketSubsystem* SocketSub = ISocketSubsystem::Get();
+	if (SocketSub)
+	{
+		double AsyncRequestStartTime = FPlatformTime::Seconds();
+		FAsyncGetAddressInfoCallback CallbackFunc = [SocketSub, AsyncRequestStartTime](FAddressInfoResult Results)
+		{
+			UE_LOG(LogSockets, Log, TEXT("Async GAI Request returned after %f seconds, started at %f"), FPlatformTime::Seconds() - AsyncRequestStartTime, AsyncRequestStartTime);
+			DebugPrintGaiResults(SocketSub, Results);
+		};
+		SocketSub->GetAddressInfoAsync(CallbackFunc, *HostStr, nullptr, EAddressInfoFlags::Default, NAME_None);
+		return;
+	}
+
+	UE_LOG(LogSockets, Warning, TEXT("Failed to get socket subsystem!"));
+}
+
+static void RunAddressSerialize(const FString& InputStr)
+{
+	if (InputStr.IsEmpty())
+	{
+		UE_LOG(LogSockets, Warning, TEXT("SOCKETSUB Serialize requires an ip address to test with."));
+		return;
+	}
+
+	ISocketSubsystem* SocketSub = ISocketSubsystem::Get();
+	if (SocketSub)
+	{
+		TSharedPtr<FInternetAddr> NewAddr = SocketSub->GetAddressFromString(InputStr);
+		if (NewAddr.IsValid())
+		{
+			UE_LOG(LogSockets, Log, TEXT("Result Address: %s Type: %s"), *NewAddr->ToString(false), *NewAddr->GetProtocolType().ToString());
+		}
+		else
+		{
+			UE_LOG(LogSockets, Warning, TEXT("Did not get results!"));
+		}
+		return;
+	}
+
+	UE_LOG(LogSockets, Warning, TEXT("Failed to get socket subsystem!"));
+}
+
+static bool SocketSubsystemCommandHandler(UWorld* InWorld, const TCHAR* Cmd, FOutputDevice& Ar)
+{
+	if (FParse::Command(&Cmd, TEXT("SOCKETSUB")))
+	{
+		if (FParse::Command(&Cmd, TEXT("GAI")))
+		{
+			RunGAIQuery(FParse::Token(Cmd, true));
+		}
+		else if (FParse::Command(&Cmd, TEXT("SERIALIZE")))
+		{
+			RunAddressSerialize(FParse::Token(Cmd, true));
+		}
+		else if (FParse::Command(&Cmd, TEXT("ASYNCGAI")))
+		{
+			RunAsyncGAIQuery(FParse::Token(Cmd, true));
+		}
+
+		return true;
+	}
+
+	return false;
+}
+
+FStaticSelfRegisteringExec FSocketSubsystemExecs(SocketSubsystemCommandHandler);
+#endif // WITH_DEV_AUTOMATION_TESTS && !UE_BUILD_SHIPPING
