Fix not setting the error code value properly when we have no data from GAI when a protocol is not unlocked for usage on a platform.

diff --git a/Engine/Source/Runtime/Sockets/Private/IPAddress.cpp b/Engine/Source/Runtime/Sockets/Private/IPAddress.cpp
--- a/Engine/Source/Runtime/Sockets/Private/IPAddress.cpp
+++ b/Engine/Source/Runtime/Sockets/Private/IPAddress.cpp
@@ -60,6 +60,10 @@ void FResolveInfoAsync::DoWork()
 
 			return;
 		}
+		else
+		{
+			ErrorCode = SE_NO_DATA;
+		}
 
 		AttemptCount++;
 	}
