Add functionality to prevent interface scopes from being an invalid number during multicast group joins on iOS.

diff --git a/Engine/Source/Runtime/Sockets/Private/IOS/SocketsBSDIOS.h b/Engine/Source/Runtime/Sockets/Private/IOS/SocketsBSDIOS.h
--- a/Engine/Source/Runtime/Sockets/Private/IOS/SocketsBSDIOS.h
+++ b/Engine/Source/Runtime/Sockets/Private/IOS/SocketsBSDIOS.h
@@ -3,6 +3,7 @@
 #pragma once
 
 #include "CoreMinimal.h"
+#include "IPAddressBSDIOS.h"
 #include "BSDSockets/SocketSubsystemBSDPrivate.h"
 #include "Sockets.h"
 #include "BSDSockets/SocketsBSD.h"
@@ -23,8 +24,30 @@ public:
 	}
 
 	virtual ~FSocketBSDIOS()
-	{	
+	{
 		FSocketBSD::Close();
 	}
 
+	virtual bool JoinMulticastGroup(const FInternetAddr& GroupAddress, const FInternetAddr& InterfaceAddress) override
+	{
+		const FInternetAddrBSDIOS& IOSInterfaceAddr = static_cast<const FInternetAddrBSDIOS&>(InterfaceAddress);
+		// iOS does not allow usage of an interface id of 0, so redirect to using the group address instead
+		if (GroupAddress.GetProtocolType() == FNetworkProtocolTypes::IPv6 && IOSInterfaceAddr.GetScopeId() == 0)
+		{
+			return FSocketBSD::JoinMulticastGroup(GroupAddress);
+		}
+
+		return FSocketBSD::JoinMulticastGroup(GroupAddress, InterfaceAddress);
+	}
+
+	virtual bool LeaveMulticastGroup(const FInternetAddr& GroupAddress, const FInternetAddr& InterfaceAddress) override
+	{
+		const FInternetAddrBSDIOS& IOSInterfaceAddr = static_cast<const FInternetAddrBSDIOS&>(InterfaceAddress);
+		if (GroupAddress.GetProtocolType() == FNetworkProtocolTypes::IPv6 && IOSInterfaceAddr.GetScopeId() == 0)
+		{
+			return FSocketBSD::LeaveMulticastGroup(GroupAddress);
+		}
+
+		return FSocketBSD::LeaveMulticastGroup(GroupAddress, InterfaceAddress);
+	}
 };
