Fixing the GetLocalHostAddr for Android such that it does not give priority to loopback addresses on Nougat based Android OS versions.

diff --git a/Engine/Source/Runtime/Sockets/Private/Android/SocketSubsystemAndroid.cpp b/Engine/Source/Runtime/Sockets/Private/Android/SocketSubsystemAndroid.cpp
--- a/Engine/Source/Runtime/Sockets/Private/Android/SocketSubsystemAndroid.cpp
+++ b/Engine/Source/Runtime/Sockets/Private/Android/SocketSubsystemAndroid.cpp
@@ -115,7 +115,7 @@ TSharedRef<FInternetAddr> FSocketSubsystemAndroid::GetLocalHostAddr(FOutputDevic
 	//
 	// Also NOTE: Network can flip out behind applications back when connectivity changes. eg. Move out of wifi range.
 	// This seems to recover OK between matches as subsystems are reinited each session Host/Join.
-	if (Addr->GetProtocolType() == FNetworkProtocolTypes::IPv4 && Addr->GetRawIp()[0] == 127)
+	if (Addr->GetProtocolType() == FNetworkProtocolTypes::IPv4 && Addr->GetRawIp()[0] != 127)
 	{
 		return Addr;
 	}
