Fix Windows compilation errors when enabling IPv6 due to platform differences.

#rnx

diff --git a/Engine/Source/Runtime/Sockets/Private/BSDIPv6Sockets/IPAddressBSDIPv6.cpp b/Engine/Source/Runtime/Sockets/Private/BSDIPv6Sockets/IPAddressBSDIPv6.cpp
--- a/Engine/Source/Runtime/Sockets/Private/BSDIPv6Sockets/IPAddressBSDIPv6.cpp
+++ b/Engine/Source/Runtime/Sockets/Private/BSDIPv6Sockets/IPAddressBSDIPv6.cpp
@@ -232,7 +232,7 @@ void FInternetAddrBSDIPv6::SetAnyAddress()
 void FInternetAddrBSDIPv6::SetBroadcastAddress()
 {
 	// broadcast means something different in IPv6, but this is a rough equivalent
-#ifndef in6addr_allnodesonlink
+#if !defined(in6addr_allnodesonlink) && !PLATFORM_WINDOWS
 	// see RFC 4291, link-local multicast address http://tools.ietf.org/html/rfc4291
 	static in6_addr in6addr_allnodesonlink =
 	{
diff --git a/Engine/Source/Runtime/Sockets/Private/BSDSockets/IPAddressBSD.cpp b/Engine/Source/Runtime/Sockets/Private/BSDSockets/IPAddressBSD.cpp
--- a/Engine/Source/Runtime/Sockets/Private/BSDSockets/IPAddressBSD.cpp
+++ b/Engine/Source/Runtime/Sockets/Private/BSDSockets/IPAddressBSD.cpp
@@ -380,7 +380,7 @@ void FInternetAddrBSD::SetIPv6BroadcastAddress()
 	Clear();
 #if PLATFORM_HAS_BSD_IPV6_SOCKETS
 	// broadcast means something different in IPv6, but this is a rough equivalent
-#ifndef in6addr_allnodesonlink
+#if !defined(in6addr_allnodesonlink) && !PLATFORM_WINDOWS
 		// see RFC 4291, link-local multicast address http://tools.ietf.org/html/rfc4291
 	static in6_addr in6addr_allnodesonlink =
 	{
