Fix Hololens to not use the IPv6 specific classes of the socket framework.

diff --git a/Engine/Source/Runtime/Sockets/Private/HoloLens/SocketSubsystemHoloLens.cpp b/Engine/Source/Runtime/Sockets/Private/HoloLens/SocketSubsystemHoloLens.cpp
--- a/Engine/Source/Runtime/Sockets/Private/HoloLens/SocketSubsystemHoloLens.cpp
+++ b/Engine/Source/Runtime/Sockets/Private/HoloLens/SocketSubsystemHoloLens.cpp
@@ -3,12 +3,13 @@
 #include "SocketSubsystemHoloLens.h"
 #include "SocketSubsystemModule.h"
 #include "Modules/ModuleManager.h"
-#include "Misc/ScopeLock.h"
 
-#include "BSDIPv6Sockets/IPAddressBSDIPv6.h"
+#include "BSDSockets/IPAddressBSD.h"
 
 #include "AllowWindowsPlatformTypes.h"
+THIRD_PARTY_INCLUDES_START
 #include "Iphlpapi.h"
+THIRD_PARTY_INCLUDES_END
 #include "HideWindowsPlatformTypes.h"
 
 
@@ -70,17 +71,15 @@ void FSocketSubsystemHoloLens::Destroy()
 /* FSocketSubsystemBSD overrides
 *****************************************************************************/
 
-FSocket* FSocketSubsystemHoloLens::CreateSocket(const FName& SocketType, const FString& SocketDescription, bool bForceUDP)
+FSocket* FSocketSubsystemHoloLens::CreateSocket(const FName& SocketType, const FString& SocketDescription, const FName& ProtocolType)
 {
-	FSocketBSD* NewSocket = static_cast<FSocketBSD*>(FSocketSubsystemBSD::CreateSocket(SocketType, SocketDescription, bForceUDP));
+	FSocketBSD* NewSocket = static_cast<FSocketBSD*>(FSocketSubsystemBSD::CreateSocket(SocketType, SocketDescription, ProtocolType));
 
-	if (NewSocket)
+	if (NewSocket != nullptr)
 	{
-		// Xbox One requires IPv6 only mode to be false
 		NewSocket->SetIPv6Only(false);
 	}
-
-	if (NewSocket == nullptr)
+	else	
 	{
 		UE_LOG(LogSockets, Warning, TEXT("Failed to create socket %s [%s]"), *SocketType.ToString(), *SocketDescription);
 	}
@@ -88,52 +87,6 @@ FSocket* FSocketSubsystemHoloLens::CreateSocket(const FName& SocketType, const F
 	return NewSocket;
 }
 
-ESocketErrors FSocketSubsystemHoloLens::GetHostByName(const ANSICHAR* HostName, FInternetAddr& OutAddr)
-{
-	FScopeLock ScopeLock(&HostByNameSynch);
-	addrinfo* AddrInfo = nullptr;
-
-	// We will allow either IPv6 or IPv4 since we call SetIPv6Only(false)
-	addrinfo HintAddrInfo;
-	FMemory::Memzero(&HintAddrInfo, sizeof(HintAddrInfo));
-	HintAddrInfo.ai_family = AF_UNSPEC;
-
-	OutAddr.SetIp(0);
-	int32 ErrorCode = getaddrinfo(HostName, nullptr, &HintAddrInfo, &AddrInfo);
-	ESocketErrors SocketError = TranslateGAIErrorCode(ErrorCode);
-	if (SocketError == SE_NO_ERROR)
-	{
-		for (; AddrInfo != nullptr; AddrInfo = AddrInfo->ai_next)
-		{
-			if (AddrInfo->ai_family == AF_INET6)
-			{
-				sockaddr_in6* IPv6SockAddr = reinterpret_cast<sockaddr_in6*>(AddrInfo->ai_addr);
-				if (IPv6SockAddr != nullptr)
-				{
-					static_cast<FInternetAddrBSDIPv6&>(OutAddr).SetIp(IPv6SockAddr->sin6_addr);
-
-					// Break out immediately if we find a v6 address - this is our preference.
-					break;
-				}
-			}
-			else if (AddrInfo->ai_family == AF_INET && !OutAddr.IsValid())
-			{
-				sockaddr_in* IPv4SockAddr = reinterpret_cast<sockaddr_in*>(AddrInfo->ai_addr);
-				if (IPv4SockAddr != nullptr)
-				{
-					static_cast<FInternetAddrBSDIPv6&>(OutAddr).SetIp(IPv4SockAddr->sin_addr);
-
-					// Keep looking in case there's a v6 address to be had
-				}
-			}
-		}
-		freeaddrinfo(AddrInfo);
-		return OutAddr.IsValid() ? SE_NO_ERROR : SE_HOST_NOT_FOUND;
-	}
-	return SocketError;
-
-}
-
 
 bool FSocketSubsystemHoloLens::Init(FString& Error)
 {
@@ -177,57 +130,6 @@ ESocketErrors FSocketSubsystemHoloLens::GetLastErrorCode()
 	return TranslateErrorCode(WSAGetLastError());
 }
 
-
-bool FSocketSubsystemHoloLens::GetLocalAdapterAddresses(TArray<TSharedPtr<FInternetAddr> >& OutAdresses)
-{
-	//ULONG Flags = GAA_FLAG_INCLUDE_PREFIX | GAA_FLAG_SKIP_MULTICAST | GAA_FLAG_SKIP_DNS_SERVER | GAA_FLAG_SKIP_FRIENDLY_NAME;
-	//ULONG Result;
-	//ULONG Size;
-
-	//// determine the required size of the address list buffer
-	//Result = GetAdaptersAddresses(AF_INET, Flags, NULL, NULL, &Size);
-
-	//if (Result != ERROR_BUFFER_OVERFLOW)
-	//{
-	//    return false;
-	//}
-
-	//PIP_ADAPTER_ADDRESSES AdapterAddresses = (PIP_ADAPTER_ADDRESSES)FMemory::Malloc(Size);
-
-	//// get the actual list of adapters
-	//Result = GetAdaptersAddresses(AF_INET, Flags, NULL, AdapterAddresses, &Size);
-
-	//if (Result != ERROR_SUCCESS)
-	//{
-	//    FMemory::Free(AdapterAddresses);
-
-	//    return false;
-	//}
-
-	//// extract the list of physical addresses from each adapter
-	//for (PIP_ADAPTER_ADDRESSES AdapterAddress = AdapterAddresses; AdapterAddress != NULL; AdapterAddress = AdapterAddress->Next)
-	//{
-	//    if ((AdapterAddress->IfType == IF_TYPE_ETHERNET_CSMACD) ||
-	//        (AdapterAddress->IfType == IF_TYPE_IEEE80211))
-	//    {
-	//        for (PIP_ADAPTER_UNICAST_ADDRESS UnicastAddress = AdapterAddress->FirstUnicastAddress; UnicastAddress != NULL; UnicastAddress = UnicastAddress->Next)
-	//        {
-	//            if ((UnicastAddress->Flags & IP_ADAPTER_ADDRESS_DNS_ELIGIBLE) != 0)
-	//            {
-	//                uint32 RawAddress = ntohl(*((uint32*)(&UnicastAddress->Address.lpSockaddr->sa_data[2])));
-
-	//                OutAdresses.Add(CreateInternetAddr(RawAddress));
-	//            }
-	//        }
-	//    }
-	//}
-
-	//FMemory::Free(AdapterAddresses);
-
-	return false;
-}
-
-
 ESocketErrors FSocketSubsystemHoloLens::TranslateErrorCode(int32 Code)
 {
 	// handle the generic -1 error
diff --git a/Engine/Source/Runtime/Sockets/Private/HoloLens/SocketSubsystemHoloLens.h b/Engine/Source/Runtime/Sockets/Private/HoloLens/SocketSubsystemHoloLens.h
--- a/Engine/Source/Runtime/Sockets/Private/HoloLens/SocketSubsystemHoloLens.h
+++ b/Engine/Source/Runtime/Sockets/Private/HoloLens/SocketSubsystemHoloLens.h
@@ -7,7 +7,6 @@
 #include "BSDSockets/SocketsBSD.h"
 #include "SocketSubsystemPackage.h"
 
-
 /**
  * HoloLens specific socket subsystem implementation
  */
@@ -28,16 +27,19 @@ public:
 
 	// FSocketSubsystemBSD overrides
 
-	virtual class FSocket* CreateSocket(const FName& SocketType, const FString& SocketDescription, bool bForceUDP = false) override;
-	virtual ESocketErrors GetHostByName(const ANSICHAR* HostName, FInternetAddr& OutAddr) override;
+	virtual class FSocket* CreateSocket(const FName& SocketType, const FString& SocketDescription, const FName& ProtocolType) override;
 	virtual bool HasNetworkDevice() override;
 	virtual ESocketErrors GetLastErrorCode() override;
-	virtual bool GetLocalAdapterAddresses(TArray<TSharedPtr<FInternetAddr> >& OutAdresses) override;
 	virtual const TCHAR* GetSocketAPIName() const override;
 	virtual bool Init(FString& Error) override;
 	virtual void Shutdown() override;
 	virtual ESocketErrors TranslateErrorCode(int32 Code) override;
 
+	virtual FName GetDefaultSocketProtocolFamily() const override
+	{
+		return FNetworkProtocolTypes::IPv6;
+	}
+
 PACKAGE_SCOPE:
 
 	/**
@@ -57,9 +59,4 @@ protected:
 
 	/** Holds the single instantiation of this subsystem. */
 	static FSocketSubsystemHoloLens* SocketSingleton;
-
-private:
-
-	// Used to prevent multiple threads accessing the shared data.
-	FCriticalSection HostByNameSynch;
 };
\ No newline at end of file
