Delete IPv6 specific classes as they've been deprecated for several engine versions now.

diff --git a/Engine/Source/Runtime/Sockets/Private/BSDIPv6Sockets/IPAddressBSDIPv6.cpp b/Engine/Source/Runtime/Sockets/Private/BSDIPv6Sockets/IPAddressBSDIPv6.cpp
deleted file mode 100644
diff --git a/Engine/Source/Runtime/Sockets/Private/BSDIPv6Sockets/IPAddressBSDIPv6.h b/Engine/Source/Runtime/Sockets/Private/BSDIPv6Sockets/IPAddressBSDIPv6.h
deleted file mode 100644
diff --git a/Engine/Source/Runtime/Sockets/Private/BSDIPv6Sockets/SocketSubsystemBSDIPv6.cpp b/Engine/Source/Runtime/Sockets/Private/BSDIPv6Sockets/SocketSubsystemBSDIPv6.cpp
deleted file mode 100644
diff --git a/Engine/Source/Runtime/Sockets/Private/BSDIPv6Sockets/SocketSubsystemBSDIPv6.h b/Engine/Source/Runtime/Sockets/Private/BSDIPv6Sockets/SocketSubsystemBSDIPv6.h
deleted file mode 100644
diff --git a/Engine/Source/Runtime/Sockets/Private/BSDIPv6Sockets/SocketsBSDIPv6.cpp b/Engine/Source/Runtime/Sockets/Private/BSDIPv6Sockets/SocketsBSDIPv6.cpp
deleted file mode 100644
diff --git a/Engine/Source/Runtime/Sockets/Private/BSDIPv6Sockets/SocketsBSDIPv6.h b/Engine/Source/Runtime/Sockets/Private/BSDIPv6Sockets/SocketsBSDIPv6.h
deleted file mode 100644
diff --git a/Engine/Source/Runtime/Sockets/Private/BSDSockets/SocketSubsystemBSDPrivate.h b/Engine/Source/Runtime/Sockets/Private/BSDSockets/SocketSubsystemBSDPrivate.h
--- a/Engine/Source/Runtime/Sockets/Private/BSDSockets/SocketSubsystemBSDPrivate.h
+++ b/Engine/Source/Runtime/Sockets/Private/BSDSockets/SocketSubsystemBSDPrivate.h
@@ -75,48 +75,3 @@ inline int TranslateFlags(ESocketReceiveFlags::Type Flags)
 
 	return TranslatedFlags;
 }
-
-
-/**
- * Standard BSD specific socket subsystem implementation (common to both IPv4 and IPv6)
- */
-class UE_DEPRECATED(4.21, "Use FSocketSubsystemBSD.") FSocketSubsystemBSDCommon
-	: public ISocketSubsystem
-{
-protected:
-
-	/**
-	 * Translates return values of getaddrinfo() to socket error enum
-	 */
-	inline ESocketErrors TranslateGAIErrorCode(int32 Code) const
-	{
-#if PLATFORM_HAS_BSD_SOCKET_FEATURE_GETADDRINFO
-		switch (Code)
-		{
-			// getaddrinfo() has its own error codes
-			case EAI_AGAIN:			return SE_TRY_AGAIN;
-			case EAI_BADFLAGS:		return SE_EINVAL;
-			case EAI_FAIL:			return SE_NO_RECOVERY;
-			case EAI_FAMILY:		return SE_EAFNOSUPPORT;
-			case EAI_MEMORY:		return SE_ENOBUFS;
-			case EAI_NONAME:		return SE_HOST_NOT_FOUND;
-			case EAI_SERVICE:		return SE_EPFNOSUPPORT;
-			case EAI_SOCKTYPE:		return SE_ESOCKTNOSUPPORT;
-#if PLATFORM_HAS_BSD_SOCKET_FEATURE_WINSOCKETS
-			case WSANO_DATA:		return SE_NO_DATA;
-			case WSANOTINITIALISED: return SE_NOTINITIALISED;
-#else			
-			case EAI_NODATA:		return SE_NO_DATA;
-			case EAI_ADDRFAMILY:	return SE_ADDRFAMILY;
-			case EAI_SYSTEM:		return SE_SYSTEM;
-#endif
-			case 0:					break; // 0 means success
-			default:
-				UE_LOG(LogSockets, Warning, TEXT("Unhandled getaddrinfo() socket error! Code: %d"), Code);
-				return SE_EINVAL;
-		}
-#endif // PLATFORM_HAS_BSD_SOCKET_FEATURE_GETADDRINFO
-
-		return SE_NO_ERROR;
-	};
-};
