PR #6397: Fixed MULTIHOME command line argument in SocketSubsystemUnix (Contributed by [REDACTED_PERSON])

diff --git a/Engine/Source/Runtime/Sockets/Private/Unix/SocketSubsystemUnix.cpp b/Engine/Source/Runtime/Sockets/Private/Unix/SocketSubsystemUnix.cpp
--- a/Engine/Source/Runtime/Sockets/Private/Unix/SocketSubsystemUnix.cpp
+++ b/Engine/Source/Runtime/Sockets/Private/Unix/SocketSubsystemUnix.cpp
@@ -111,10 +111,18 @@ TSharedRef<FInternetAddr> FSocketSubsystemUnix::GetLocalHostAddr(FOutputDevice&
 {
 	bCanBindAll = true;
 
+	TSharedRef<FInternetAddr> MultihomeAddress = CreateInternetAddr();
+	if (GetMultihomeAddress(MultihomeAddress))
+	{
+		bCanBindAll = false;
+		UE_LOG(LogSockets, Verbose, TEXT("Local address is %s"), *(MultihomeAddress->ToString(false)));
+		return MultihomeAddress->Clone();
+	}
+
 	TArray<TSharedPtr<FInternetAddr>> ResultArray;
 	if (GetLocalAdapterAddresses(ResultArray))
 	{
-		if (FParse::Param(FCommandLine::Get(), TEXT("PRIMARYNET")) || FParse::Param(FCommandLine::Get(), TEXT("MULTIHOME")))
+		if (FParse::Param(FCommandLine::Get(), TEXT("PRIMARYNET")))
 		{
 			bCanBindAll = false;
 		}
