Use new inline versions of substring functions
Fix up nearby cases where ESearchCase::CaseSensitive should have been used
#rnx

diff --git a/Engine/Source/Runtime/Sockets/Private/BSDIPv6Sockets/IPAddressBSDIPv6.cpp b/Engine/Source/Runtime/Sockets/Private/BSDIPv6Sockets/IPAddressBSDIPv6.cpp
--- a/Engine/Source/Runtime/Sockets/Private/BSDIPv6Sockets/IPAddressBSDIPv6.cpp
+++ b/Engine/Source/Runtime/Sockets/Private/BSDIPv6Sockets/IPAddressBSDIPv6.cpp
@@ -91,7 +91,7 @@ void FInternetAddrBSDIPv6::SetIp(const TCHAR* InAddr, bool& bIsValid)
 		if (AddressString.Contains("]:") || (!bIsIPv6 && LastColonIndex != INDEX_NONE))
 		{
 			Port = AddressString.RightChop(LastColonIndex + 1);
-			AddressString = AddressString.Left(LastColonIndex);
+			AddressString.LeftInline(LastColonIndex, false);
 		}
 
 		AddressString.RemoveFromStart("[");
