Fixed Launch On crash on iOS when WiFi and cellular are disabled

diff --git a/Engine/Source/Runtime/Sockets/Private/IOS/SocketSubsystemIOS.cpp b/Engine/Source/Runtime/Sockets/Private/IOS/SocketSubsystemIOS.cpp
--- a/Engine/Source/Runtime/Sockets/Private/IOS/SocketSubsystemIOS.cpp
+++ b/Engine/Source/Runtime/Sockets/Private/IOS/SocketSubsystemIOS.cpp
@@ -222,7 +222,7 @@ bool FSocketSubsystemIOS::GetLocalAdapterAddresses(TArray<TSharedPtr<FInternetAd
 			OutAddresses.Add(AddressInterface.Address);
 		}
 
-		return true;
+		return OutAddresses.Num() > 0;
 	}
 	return false;
 }
