Add GetLocalHostAddrViaConnect

This new seems to report a more correct IP address via GetLocalHostAddr.  If this method doesn't succeed the previous mechanisms are used.

#rnx

diff --git a/Engine/Source/Runtime/Sockets/Private/SocketSubsystem.cpp b/Engine/Source/Runtime/Sockets/Private/SocketSubsystem.cpp
--- a/Engine/Source/Runtime/Sockets/Private/SocketSubsystem.cpp
+++ b/Engine/Source/Runtime/Sockets/Private/SocketSubsystem.cpp
@@ -11,6 +11,7 @@
 #include "IPAddress.h"
 #include "Sockets.h"
 #include "Templates/UniquePtr.h"
+#include "BSDSockets/IPAddressBSD.h"
 
 DEFINE_LOG_CATEGORY(LogSockets);
 
@@ -75,6 +76,44 @@ static IModuleInterface* LoadSubsystemModule(const FString& SubsystemName)
 	return nullptr;
 }
 
+/**
+ * Helper function that attempts to determine the host address via socket connect
+ *
+ * @param IO The address if successful
+ * @return true if successful, false otherwise
+ */
+static bool GetLocalHostAddrViaConnect(TSharedRef<FInternetAddr>& IO)
+{
+	sockaddr_in Addr;
+	Addr.sin_family = AF_INET;
+	Addr.sin_addr.s_addr = 0xffff1fac; // any IP will do, doesn't even need to be reachable
+	Addr.sin_port = 0x0100;
+	socklen_t AddrSize = sizeof(Addr);
+
+	int Socket = socket(AF_INET, SOCK_DGRAM, 0);
+
+	if ( connect(Socket, (struct sockaddr*)&Addr, sizeof(Addr)) != 0 )
+	{
+		close(Socket);
+		return false;
+	}
+	if ( getsockname(Socket, (struct sockaddr*)&Addr, &AddrSize) != 0 )
+	{
+		close(Socket);
+		return false;
+	}
+	
+	Addr.sin_port = 0;
+
+	IO = MakeShareable(new FInternetAddrBSD());
+	StaticCastSharedRef<FInternetAddrBSD>(IO)->Set(*reinterpret_cast<sockaddr_storage*>(&Addr), AddrSize);
+
+	UE_LOG(LogSockets, VeryVerbose, TEXT("GetLocalHostAddrViaConnect: %s"), *IO->ToString(true));
+	
+	close(Socket);
+	
+	return true;
+}
 
 FUniqueSocket ISocketSubsystem::CreateUniqueSocket(const FName& SocketType, const FString& SocketDescription, bool bForceUDP)
 {
@@ -408,23 +447,29 @@ TSharedRef<FInternetAddr> ISocketSubsystem::GetLocalHostAddr(FOutputDevice& Out,
 
 	if (!GetMultihomeAddress(HostAddr))
 	{
-		bCanBindAll = true;
-
-		TArray<TSharedPtr<FInternetAddr>> AdapterAddresses;
-		if (!GetLocalAdapterAddresses(AdapterAddresses) || (AdapterAddresses.Num() == 0))
-		{
-			Out.Logf(TEXT("Could not fetch the local adapter addresses"));
-			HostAddr->SetAnyAddress();
-		}
-		else
+		// First try to get the host address via connect
+		if (!GetLocalHostAddrViaConnect(HostAddr))
 		{
-			if (AdapterAddresses.Num() > 0)
-			{ 
-				HostAddr = AdapterAddresses[0]->Clone();
+			bCanBindAll = true;
+
+			TArray<TSharedPtr<FInternetAddr>> AdapterAddresses;
+			if (!GetLocalAdapterAddresses(AdapterAddresses) || (AdapterAddresses.Num() == 0))
+			{
+				Out.Logf(TEXT("Could not fetch the local adapter addresses"));
+				HostAddr->SetAnyAddress();
+			}
+			else
+			{
+				if (AdapterAddresses.Num() > 0)
+				{
+					HostAddr = AdapterAddresses[0]->Clone();
+				}
 			}
 		}
 	}
 
+	UE_LOG(LogSockets, VeryVerbose, TEXT("GetLocalHostAddr: %s"), *HostAddr->ToString(true));
+
 	return HostAddr;
 }
 
