NetworkMessage: Added a missing virtual destructor to FSimpleAbstractSocket

#rnx

diff --git a/Engine/Source/Runtime/Sockets/Public/NetworkMessage.h b/Engine/Source/Runtime/Sockets/Public/NetworkMessage.h
--- a/Engine/Source/Runtime/Sockets/Public/NetworkMessage.h
+++ b/Engine/Source/Runtime/Sockets/Public/NetworkMessage.h
@@ -94,6 +94,8 @@ public:
 	virtual bool Send(const uint8 *Buffer, int32 Size) const = 0;
 	/** return the magic number for this message, also used for endian correction on the archives **/
 	virtual uint32 GetMagic() const = 0;
+	/** Destructor */
+	virtual ~FSimpleAbstractSocket() { }
 };
 
 /**
