#8402 : PR : Add support for using poll instead of select in BSD sockets. Sockets will use poll by default (if supported by the platform), but this can be toggled with the USE_SOCKET_FEATURE_POLL macro.

diff --git a/Engine/Source/Runtime/Core/Public/HAL/Platform.h b/Engine/Source/Runtime/Core/Public/HAL/Platform.h
--- a/Engine/Source/Runtime/Core/Public/HAL/Platform.h
+++ b/Engine/Source/Runtime/Core/Public/HAL/Platform.h
@@ -294,6 +294,9 @@
 #ifndef PLATFORM_HAS_BSD_SOCKET_FEATURE_IOCTL
 	#define PLATFORM_HAS_BSD_SOCKET_FEATURE_IOCTL	1
 #endif
+#ifndef PLATFORM_HAS_BSD_SOCKET_FEATURE_POLL
+	#define PLATFORM_HAS_BSD_SOCKET_FEATURE_POLL	0
+#endif
 #ifndef PLATFORM_HAS_BSD_SOCKET_FEATURE_SELECT
 	#define PLATFORM_HAS_BSD_SOCKET_FEATURE_SELECT	1
 #endif
diff --git a/Engine/Source/Runtime/Core/Public/Unix/UnixPlatform.h b/Engine/Source/Runtime/Core/Public/Unix/UnixPlatform.h
--- a/Engine/Source/Runtime/Core/Public/Unix/UnixPlatform.h
+++ b/Engine/Source/Runtime/Core/Public/Unix/UnixPlatform.h
@@ -60,6 +60,7 @@ typedef FUnixPlatformTypes FPlatformTypes;
 #define PLATFORM_HAS_NO_EPROCLIM						1
 #define PLATFORM_HAS_BSD_IPV6_SOCKETS					1
 #define PLATFORM_HAS_BSD_SOCKET_FEATURE_IOCTL			1
+#define PLATFORM_HAS_BSD_SOCKET_FEATURE_POLL			1
 #define PLATFORM_HAS_BSD_SOCKET_FEATURE_MSG_DONTWAIT	1
 #define PLATFORM_HAS_BSD_SOCKET_FEATURE_RECVMMSG		1
 #define PLATFORM_HAS_BSD_SOCKET_FEATURE_TIMESTAMP		1
diff --git a/Engine/Source/Runtime/Sockets/Private/BSDSockets/SocketSubsystemBSDPrivate.h b/Engine/Source/Runtime/Sockets/Private/BSDSockets/SocketSubsystemBSDPrivate.h
--- a/Engine/Source/Runtime/Sockets/Private/BSDSockets/SocketSubsystemBSDPrivate.h
+++ b/Engine/Source/Runtime/Sockets/Private/BSDSockets/SocketSubsystemBSDPrivate.h
@@ -25,6 +25,9 @@
 	#include <fcntl.h>
 	#include <sys/types.h>
 	#include <sys/ioctl.h>
+#endif
+#if PLATFORM_HAS_BSD_SOCKET_FEATURE_POLL
+	#include <poll.h>
 #endif
 	#include <netinet/in.h>
 	#include <arpa/inet.h>
diff --git a/Engine/Source/Runtime/Sockets/Private/BSDSockets/SocketsBSD.cpp b/Engine/Source/Runtime/Sockets/Private/BSDSockets/SocketsBSD.cpp
--- a/Engine/Source/Runtime/Sockets/Private/BSDSockets/SocketsBSD.cpp
+++ b/Engine/Source/Runtime/Sockets/Private/BSDSockets/SocketsBSD.cpp
@@ -4,11 +4,14 @@
 
 #if PLATFORM_HAS_BSD_SOCKETS || PLATFORM_HAS_BSD_IPV6_SOCKETS
 
+#ifndef USE_SOCKET_FEATURE_POLL
+	#define USE_SOCKET_FEATURE_POLL 1
+#endif
+
 #include "BSDSockets/IPAddressBSD.h"
 #include "BSDSockets/SocketSubsystemBSD.h"
 //#include "Net/NetworkProfiler.h"
 
-
 /* FSocket overrides
  *****************************************************************************/
 
@@ -707,7 +710,43 @@ bool FSocketBSD::SetIPv6Only(bool bIPv6Only)
 
 ESocketBSDReturn FSocketBSD::HasState(ESocketBSDParam State, FTimespan WaitTime)
 {
-#if PLATFORM_HAS_BSD_SOCKET_FEATURE_SELECT
+#if PLATFORM_HAS_BSD_SOCKET_FEATURE_POLL && USE_SOCKET_FEATURE_POLL
+	struct pollfd FDSet;
+	FDSet.fd = Socket;
+	FDSet.revents = 0;
+
+	switch (State)
+	{
+	case ESocketBSDParam::CanRead:
+		FDSet.events = POLLIN;
+		break;
+
+	case ESocketBSDParam::CanWrite:
+		FDSet.events = POLLOUT;
+		break;
+
+	case ESocketBSDParam::HasError:
+		FDSet.events = POLLERR | POLLPRI;
+		break;
+	}
+	int Result = ::poll(&FDSet, 1, WaitTime.GetTotalMilliseconds());
+	if (Result >= 0)
+	{
+		if ((FDSet.revents & FDSet.events) > 0)
+		{
+			return ESocketBSDReturn::Yes;
+		}
+		else
+		{
+			return ESocketBSDReturn::No;
+		}
+	}
+	else
+	{
+		return ESocketBSDReturn::EncounteredError;
+	}
+
+#elif PLATFORM_HAS_BSD_SOCKET_FEATURE_SELECT
 	// convert WaitTime to a timeval
 	timeval Time;
 	Time.tv_sec = (int32)WaitTime.GetTotalSeconds();
@@ -743,7 +782,7 @@ ESocketBSDReturn FSocketBSD::HasState(ESocketBSDParam State, FTimespan WaitTime)
 		SelectStatus == 0 ? ESocketBSDReturn::No :
 		ESocketBSDReturn::EncounteredError;
 #else
-	UE_LOG(LogSockets, Fatal, TEXT("This platform doesn't support select(), but FSocketBSD::HasState was not overridden"));
+	UE_LOG(LogSockets, Fatal, TEXT("This platform doesn't support poll() or select(), but FSocketBSD::HasState was not overridden"));
 	return ESocketBSDReturn::EncounteredError;
 #endif
 }
