Fix for UTF-8 array initialization.

diff --git a/Engine/Source/Runtime/Sockets/Private/SocketSubsystem.cpp b/Engine/Source/Runtime/Sockets/Private/SocketSubsystem.cpp
--- a/Engine/Source/Runtime/Sockets/Private/SocketSubsystem.cpp
+++ b/Engine/Source/Runtime/Sockets/Private/SocketSubsystem.cpp
@@ -488,7 +488,7 @@ TSharedRef<FInternetAddr> ISocketSubsystem::GetLocalBindAddr(FOutputDevice& Out)
 
 bool ISocketSubsystem::GetMultihomeAddress(TSharedRef<FInternetAddr>& Addr)
 {
-	TCHAR Home[256] = TEXT("");
+	TCHAR Home[256] = {};
 	if (FParse::Value(FCommandLine::Get(), TEXT("MULTIHOME="), Home, UE_ARRAY_COUNT(Home)))
 	{
 		TSharedPtr<FInternetAddr> MultiHomeQuery = GetAddressFromString(Home);
