Fixed redundant socket bindings on some platforms (Revision #1).

diff --git a/Engine/Source/Runtime/Sockets/Private/SocketSubsystem.cpp b/Engine/Source/Runtime/Sockets/Private/SocketSubsystem.cpp
--- a/Engine/Source/Runtime/Sockets/Private/SocketSubsystem.cpp
+++ b/Engine/Source/Runtime/Sockets/Private/SocketSubsystem.cpp
@@ -382,9 +382,13 @@ TArray<TSharedRef<FInternetAddr>> ISocketSubsystem::GetLocalBindAddresses()
 	if (BindableAddresses.ReturnCode == SE_NO_ERROR)
 	{
 		// Push in all the bindable addresses.
-		for (const auto& BindAddresses : BindableAddresses.Results)
+		for (const FAddressInfoResultData& BindAddresses : BindableAddresses.Results)
 		{
-			BindingAddresses.Add(BindAddresses.Address);
+			// GetAddressInfo can return both TCP and UDP bindings for the same address - which is redundant when returning only addresses
+			if (!BindingAddresses.ContainsByPredicate([&](const TSharedRef<FInternetAddr>& A) { return *A == *BindAddresses.Address; }))
+			{
+				BindingAddresses.Add(BindAddresses.Address);
+			}
 		}
 	}
 
