Fixed FSocket::GetConnectionState returning 'SCS_Connected' status, while the socket is closed. Fix contributed by [REDACTED_PERSON].

#9363

diff --git a/Engine/Source/Runtime/Sockets/Private/BSDSockets/SocketsBSD.cpp b/Engine/Source/Runtime/Sockets/Private/BSDSockets/SocketsBSD.cpp
--- a/Engine/Source/Runtime/Sockets/Private/BSDSockets/SocketsBSD.cpp
+++ b/Engine/Source/Runtime/Sockets/Private/BSDSockets/SocketsBSD.cpp
@@ -369,7 +369,7 @@ ESocketConnectionState FSocketBSD::GetConnectionState(void)
 	ESocketConnectionState CurrentState = SCS_ConnectionError;
 
 	// look for an existing error
-	if (HasState(ESocketBSDParam::HasError) == ESocketBSDReturn::No)
+	if (Socket != INVALID_SOCKET && HasState(ESocketBSDParam::HasError) == ESocketBSDReturn::No)
 	{
 		if (FPlatformTime::Seconds() - LastActivityTime > 5.0)
 		{
