Fixed signed vs unsigned bounds check.

diff --git a/Engine/Source/Runtime/Sockets/Public/MultichannelTcpSocket.h b/Engine/Source/Runtime/Sockets/Public/MultichannelTcpSocket.h
--- a/Engine/Source/Runtime/Sockets/Public/MultichannelTcpSocket.h
+++ b/Engine/Source/Runtime/Sockets/Public/MultichannelTcpSocket.h
@@ -59,51 +59,55 @@ public:
 	 */
 	int32 BlockingReceive( uint8* Data, int32 Count, uint32 Channel )
 	{
-		check((Channel != ControlChannel) && Data && Count); // Channel 0 is reserved, must want data
+		check((Channel != ControlChannel) && Data); // Channel 0 is reserved, must want data
 		
-		for (int32 Attempt = 0; Attempt < 2; Attempt++)
+		if (Count > 0)
 		{
-			FScopedEvent* LocalEventToRestart = NULL;
+			for (int32 Attempt = 0; Attempt < 2; Attempt++)
 			{
-				FScopeLock ScopeLock(&ReceiveBuffersCriticalSection);
+				FScopedEvent* LocalEventToRestart = NULL;
+				{
+					FScopeLock ScopeLock(&ReceiveBuffersCriticalSection);
 
-				FReceiveBuffer* ChannelBuffer = ReceiveBuffers.FindRef(Channel);
+					FReceiveBuffer* ChannelBuffer = ReceiveBuffers.FindRef(Channel);
 
-				if (!ChannelBuffer)
-				{
-					ChannelBuffer = new FReceiveBuffer();
+					if (!ChannelBuffer)
+					{
+						ChannelBuffer = new FReceiveBuffer();
 
-					ReceiveBuffers.Add(Channel, ChannelBuffer);
-				}
+						ReceiveBuffers.Add(Channel, ChannelBuffer);
+					}
 
-				check(!ChannelBuffer->EventToResumeWhenDataIsReady && !ChannelBuffer->BytesRequiredToResume); // would be bad to have multiple listeners
+					// Would be bad to have multiple listeners
+					check(!ChannelBuffer->EventToResumeWhenDataIsReady && !ChannelBuffer->BytesRequiredToResume);
 				
-				if (ChannelBuffer->Buffer.Num() >= Count)
-				{
-					FMemory::Memcpy(Data, ChannelBuffer->Buffer.GetData(), Count);
-
-					if (ChannelBuffer->Buffer.Num() == Count)
+					if (ChannelBuffer->Buffer.Num() >= Count)
 					{
-						ReceiveBuffers.Remove(Channel);
+						FMemory::Memcpy(Data, ChannelBuffer->Buffer.GetData(), Count);
+
+						if (ChannelBuffer->Buffer.Num() == Count)
+						{
+							ReceiveBuffers.Remove(Channel);
+						}
+						else
+						{
+							ChannelBuffer->Buffer.RemoveAt(0, Count);
+						}
+
+						return Count;
 					}
-					else
+
+					if (!Attempt) // if someone woke us up with insufficient data, we simply error out
 					{
-						ChannelBuffer->Buffer.RemoveAt(0, Count);
-					}
+						LocalEventToRestart = new FScopedEvent();
 
-					return Count;
+						ChannelBuffer->EventToResumeWhenDataIsReady = LocalEventToRestart;
+						ChannelBuffer->BytesRequiredToResume = Count;
+					}
 				}
 
-				if (!Attempt) // if someone woke us up with insufficient data, we simply error out
-				{
-					LocalEventToRestart = new FScopedEvent();
-
-					ChannelBuffer->EventToResumeWhenDataIsReady = LocalEventToRestart;
-					ChannelBuffer->BytesRequiredToResume = Count;
-				}
+				delete LocalEventToRestart; // wait here for sufficient data
 			}
-
-			delete LocalEventToRestart; // wait here for sufficient data
 		}
 
 		return 0;
