Fixed Receive calls with negative size, to properly return false.

diff --git a/Engine/Source/Runtime/Sockets/Private/NetworkMessage.cpp b/Engine/Source/Runtime/Sockets/Private/NetworkMessage.cpp
--- a/Engine/Source/Runtime/Sockets/Private/NetworkMessage.cpp
+++ b/Engine/Source/Runtime/Sockets/Private/NetworkMessage.cpp
@@ -8,6 +8,12 @@
 bool FSimpleAbstractSocket_FSocket::Receive(uint8 *Results, int32 Size) const
 {
 	int32 Offset = 0;
+
+	if (Size < 0)
+	{
+		return false;
+	}
+
 	while (Size > 0)
 	{
 		int32 NumRead = 0;
