Moving BSD code out of base socket subsystem
Use socket apis in GetLocalHostAddrViaConnect to use the correct types and respect ipv6

diff --git a/Engine/Source/Runtime/Sockets/Private/BSDSockets/SocketSubsystemBSD.cpp b/Engine/Source/Runtime/Sockets/Private/BSDSockets/SocketSubsystemBSD.cpp
--- a/Engine/Source/Runtime/Sockets/Private/BSDSockets/SocketSubsystemBSD.cpp
+++ b/Engine/Source/Runtime/Sockets/Private/BSDSockets/SocketSubsystemBSD.cpp
@@ -503,4 +503,81 @@ int32 FSocketSubsystemBSD::GetAddressInfoHintFlag(EAddressInfoFlags InFlags) con
 	return ReturnFlagsCode;
 }
 
+TSharedRef<FInternetAddr> FSocketSubsystemBSD::GetLocalHostAddr(FOutputDevice& Out, bool& bCanBindAll)
+{
+	TSharedRef<FInternetAddr> HostAddr = CreateInternetAddr();
+	bCanBindAll = false;
+
+	if (!GetMultihomeAddress(HostAddr))
+	{
+		// First try to get the host address via connect
+		if (!GetLocalHostAddrViaConnect(HostAddr))
+		{
+			bCanBindAll = true;
+
+			TArray<TSharedPtr<FInternetAddr>> AdapterAddresses;
+			if (!GetLocalAdapterAddresses(AdapterAddresses) || (AdapterAddresses.Num() == 0))
+			{
+				Out.Logf(TEXT("Could not fetch the local adapter addresses"));
+				HostAddr->SetAnyAddress();
+			}
+			else
+			{
+				if (AdapterAddresses.Num() > 0)
+				{
+					HostAddr = AdapterAddresses[0]->Clone();
+				}
+			}
+		}
+	}
+
+	UE_LOG(LogSockets, VeryVerbose, TEXT("GetLocalHostAddr: %s"), *HostAddr->ToString(true));
+
+	return HostAddr;
+}
+
+bool FSocketSubsystemBSD::GetLocalHostAddrViaConnect(TSharedRef<FInternetAddr>& HostAddr)
+{
+	bool bReturnValue = false;
+
+	const bool bDisableIPv6 = CVarDisableIPv6.GetValueOnAnyThread() == 1;
+
+	TSharedRef<FInternetAddr> ConnectAddr = bDisableIPv6 ? CreateInternetAddr(FNetworkProtocolTypes::IPv4) : CreateInternetAddr();
+
+	bool bIsValid = false;
+	// any IP will do, doesn't even need to be reachable
+	if (ConnectAddr->GetProtocolType() == FNetworkProtocolTypes::IPv6)
+	{
+		ConnectAddr->SetIp(TEXT("::ffff:172.31.255.255"), bIsValid);		
+	}
+	else if (ConnectAddr->GetProtocolType() == FNetworkProtocolTypes::IPv4)
+	{
+		ConnectAddr->SetIp(TEXT("172.31.255.255"), bIsValid);
+	}
+
+	if (bIsValid)
+	{
+		ConnectAddr->SetPort(256);
+
+		FUniqueSocket TestSocket = CreateUniqueSocket(NAME_DGram, TEXT("GetLocalHostAddrViaConnect"), ConnectAddr->GetProtocolType());
+
+		if (TestSocket.IsValid())
+		{
+			if (TestSocket->Connect(*ConnectAddr))
+			{
+				TestSocket->GetAddress(*HostAddr);
+
+				HostAddr->SetPort(0);
+				bReturnValue = true;
+
+				UE_LOG(LogSockets, VeryVerbose, TEXT("GetLocalHostAddrViaConnect: %s"), *HostAddr->ToString(true));
+			}
+
+			TestSocket->Close();
+		}
+	}
+
+	return bReturnValue;
+}
+
 #endif	//PLATFORM_HAS_BSD_SOCKETS
diff --git a/Engine/Source/Runtime/Sockets/Private/BSDSockets/SocketSubsystemBSD.h b/Engine/Source/Runtime/Sockets/Private/BSDSockets/SocketSubsystemBSD.h
--- a/Engine/Source/Runtime/Sockets/Private/BSDSockets/SocketSubsystemBSD.h
+++ b/Engine/Source/Runtime/Sockets/Private/BSDSockets/SocketSubsystemBSD.h
@@ -79,6 +79,8 @@ public:
 
 	virtual bool IsSocketWaitSupported() const override;
 
+	virtual TSharedRef<FInternetAddr> GetLocalHostAddr(FOutputDevice& Out, bool& bCanBindAll) override;
+
 	/**
 	 * Translates an ESocketProtocolFamily code into a value usable by raw socket apis.
 	 */
@@ -125,6 +127,14 @@ protected:
 		return InternalBSDSocketFactory(Socket, SocketType, SocketDescription, GetDefaultSocketProtocolFamily());
 	}
 
+	/**
+	 * Helper function that attempts to determine the host address via socket connect
+	 *
+	 * @param HostAddr The address if successful
+	 * @return true if successful, false otherwise
+	 */
+	bool GetLocalHostAddrViaConnect(TSharedRef<FInternetAddr>& HostAddr);
+
 	// allow BSD sockets to use this when creating new sockets from accept() etc
 	friend FSocketBSD;
 };
diff --git a/Engine/Source/Runtime/Sockets/Private/SocketSubsystem.cpp b/Engine/Source/Runtime/Sockets/Private/SocketSubsystem.cpp
--- a/Engine/Source/Runtime/Sockets/Private/SocketSubsystem.cpp
+++ b/Engine/Source/Runtime/Sockets/Private/SocketSubsystem.cpp
@@ -11,7 +11,6 @@
 #include "IPAddress.h"
 #include "Sockets.h"
 #include "Templates/UniquePtr.h"
-#include "BSDSockets/IPAddressBSD.h"
 
 DEFINE_LOG_CATEGORY(LogSockets);
 
@@ -76,46 +75,6 @@ static IModuleInterface* LoadSubsystemModule(const FString& SubsystemName)
 	return nullptr;
 }
 
-/**
- * Helper function that attempts to determine the host address via socket connect
- *
- * @param IO The address if successful
- * @return true if successful, false otherwise
- */
-static bool GetLocalHostAddrViaConnect(TSharedRef<FInternetAddr>& IO)
-{
-	sockaddr_in Addr;
-	Addr.sin_family = AF_INET;
-	Addr.sin_addr.s_addr = 0xffff1fac; // any IP will do, doesn't even need to be reachable
-	Addr.sin_port = 0x0100;
-	socklen_t AddrSize = sizeof(Addr);
-
-	// The type returned by socket varies by platform.
-	auto Socket = socket(AF_INET, SOCK_DGRAM, 0);
-
-	if ( connect(Socket, (struct sockaddr*)&Addr, sizeof(Addr)) != 0 )
-	{
-		closesocket(Socket);
-		return false;
-	}
-	if ( getsockname(Socket, (struct sockaddr*)&Addr, &AddrSize) != 0 )
-	{
-		closesocket(Socket);
-		return false;
-	}
-	
-	Addr.sin_port = 0;
-
-	IO = MakeShareable(new FInternetAddrBSD());
-	StaticCastSharedRef<FInternetAddrBSD>(IO)->Set(*reinterpret_cast<sockaddr_storage*>(&Addr), AddrSize);
-
-	UE_LOG(LogSockets, VeryVerbose, TEXT("GetLocalHostAddrViaConnect: %s"), *IO->ToString(true));
-	
-	closesocket(Socket);
-	
-	return true;
-}
-
 FUniqueSocket ISocketSubsystem::CreateUniqueSocket(const FName& SocketType, const FString& SocketDescription, bool bForceUDP)
 {
 	return FUniqueSocket(CreateSocket(SocketType, SocketDescription, bForceUDP), FSocketDeleter(this));
@@ -126,7 +85,6 @@ FUniqueSocket ISocketSubsystem::CreateUniqueSocket(const FName& SocketType, cons
 	return FUniqueSocket(CreateSocket(SocketType, SocketDescription, ProtocolName), FSocketDeleter(this));
 }
 
-
 /**
  * Shutdown all registered subsystems
  */
@@ -452,29 +410,23 @@ TSharedRef<FInternetAddr> ISocketSubsystem::GetLocalHostAddr(FOutputDevice& Out,
 
 	if (!GetMultihomeAddress(HostAddr))
 	{
-		// First try to get the host address via connect
-		if (!GetLocalHostAddrViaConnect(HostAddr))
-		{
-			bCanBindAll = true;
+		bCanBindAll = true;
 
-			TArray<TSharedPtr<FInternetAddr>> AdapterAddresses;
-			if (!GetLocalAdapterAddresses(AdapterAddresses) || (AdapterAddresses.Num() == 0))
-			{
-				Out.Logf(TEXT("Could not fetch the local adapter addresses"));
-				HostAddr->SetAnyAddress();
-			}
-			else
+		TArray<TSharedPtr<FInternetAddr>> AdapterAddresses;
+		if (!GetLocalAdapterAddresses(AdapterAddresses) || (AdapterAddresses.Num() == 0))
+		{
+			Out.Logf(TEXT("Could not fetch the local adapter addresses"));
+			HostAddr->SetAnyAddress();
+		}
+		else
+		{
+			if (AdapterAddresses.Num() > 0)
 			{
-				if (AdapterAddresses.Num() > 0)
-				{
-					HostAddr = AdapterAddresses[0]->Clone();
-				}
+				HostAddr = AdapterAddresses[0]->Clone();
 			}
 		}
 	}
 
-	UE_LOG(LogSockets, VeryVerbose, TEXT("GetLocalHostAddr: %s"), *HostAddr->ToString(true));
-
 	return HostAddr;
 }
 
