COTF/Sockets: Fix an issue with hostname to IP resolution on MacOS.

diff --git a/Engine/Source/Editor/UnrealEd/Private/CookOnTheFlyServer.cpp b/Engine/Source/Editor/UnrealEd/Private/CookOnTheFlyServer.cpp
--- a/Engine/Source/Editor/UnrealEd/Private/CookOnTheFlyServer.cpp
+++ b/Engine/Source/Editor/UnrealEd/Private/CookOnTheFlyServer.cpp
@@ -652,6 +652,10 @@ bool UCookOnTheFlyServer::StartCookOnTheFly(FCookOnTheFlyStartupOptions InCookOn
 		{
 			UE_LOG(LogCook, Display, TEXT("Unreal Network File Server is ready for client connections on %s!"), *ListenAddresses[0]->ToString(true));
 		}
+		else
+		{
+			UE_LOG(LogCook, Fatal, TEXT("Unable to get any ListenAddresses for Unreal Network file server!"));
+		}
 	}
 	else
 	{
diff --git a/Engine/Source/Runtime/Sockets/Private/Mac/SocketSubsystemMac.cpp b/Engine/Source/Runtime/Sockets/Private/Mac/SocketSubsystemMac.cpp
--- a/Engine/Source/Runtime/Sockets/Private/Mac/SocketSubsystemMac.cpp
+++ b/Engine/Source/Runtime/Sockets/Private/Mac/SocketSubsystemMac.cpp
@@ -3,7 +3,10 @@
 #include "SocketSubsystemMac.h"
 #include "SocketSubsystemModule.h"
 #include "Modules/ModuleManager.h"
+#include "BSDSockets/IPAddressBSD.h"
 
+#include <ifaddrs.h>
+#include <net/if.h>
 
 FSocketSubsystemMac* FSocketSubsystemMac::SocketSingleton = NULL;
 
@@ -78,3 +81,47 @@ class FSocketBSD* FSocketSubsystemMac::InternalBSDSocketFactory(SOCKET Socket, E
 	}
 	return MacSock;
 }
+
+bool FSocketSubsystemMac::GetLocalAdapterAddresses(TArray<TSharedPtr<FInternetAddr>>& OutAddresses)
+{
+	ifaddrs* Interfaces = NULL;
+	int InterfaceQueryRet = getifaddrs(&Interfaces);
+	UE_LOG(LogSockets, Verbose, TEXT("Querying net interfaces returned: %d"), InterfaceQueryRet);
+	const bool bDisableIPv6 = CVarDisableIPv6.GetValueOnAnyThread() == 1;
+	if (InterfaceQueryRet == 0)
+	{
+		// Loop through linked list of interfaces
+		for (ifaddrs* Travel = Interfaces; Travel != NULL; Travel = Travel->ifa_next)
+		{
+			// Skip over empty data sets.
+			if (Travel->ifa_addr == NULL)
+			{
+				continue;
+			}
+
+			uint16 AddrFamily = Travel->ifa_addr->sa_family;
+			// Find any up and non-loopback addresses
+			if ((Travel->ifa_flags & IFF_UP) != 0 &&
+				(Travel->ifa_flags & IFF_LOOPBACK) == 0 && 
+				(AddrFamily == AF_INET || (!bDisableIPv6 && AddrFamily == AF_INET6)))
+			{
+				TSharedRef<FInternetAddrBSD> NewAddress = MakeShareable(new FInternetAddrBSD(this));
+				NewAddress->SetIp(*((sockaddr_storage*)Travel->ifa_addr));
+				uint32 AddressInterface = ntohl(if_nametoindex(Travel->ifa_name));
+
+				NewAddress->SetScopeId(AddressInterface);
+				OutAddresses.Add(NewAddress);
+				UE_LOG(LogSockets, Verbose, TEXT("Added address %s on interface %d"), *(NewAddress->ToString(false)), AddressInterface);
+			}
+		}
+
+		freeifaddrs(Interfaces);
+	}
+	else
+	{
+		UE_LOG(LogSockets, Warning, TEXT("getifaddrs returned result %d"), InterfaceQueryRet);
+		return false;
+	}
+
+	return (OutAddresses.Num() > 0);
+}
diff --git a/Engine/Source/Runtime/Sockets/Private/Mac/SocketSubsystemMac.h b/Engine/Source/Runtime/Sockets/Private/Mac/SocketSubsystemMac.h
--- a/Engine/Source/Runtime/Sockets/Private/Mac/SocketSubsystemMac.h
+++ b/Engine/Source/Runtime/Sockets/Private/Mac/SocketSubsystemMac.h
@@ -67,6 +67,8 @@ public:
 	 * @return Whether the machine has a properly configured network device or not
 	 */
 	virtual bool HasNetworkDevice() override;
+
+	virtual bool GetLocalAdapterAddresses(TArray<TSharedPtr<FInternetAddr>>& OutAddresses) override;
 };
 
 /**
diff --git a/Engine/Source/Runtime/Sockets/Private/SocketSubsystem.cpp b/Engine/Source/Runtime/Sockets/Private/SocketSubsystem.cpp
--- a/Engine/Source/Runtime/Sockets/Private/SocketSubsystem.cpp
+++ b/Engine/Source/Runtime/Sockets/Private/SocketSubsystem.cpp
@@ -357,6 +357,9 @@ TArray<TSharedRef<FInternetAddr>> ISocketSubsystem::GetLocalBindAddresses()
 bool ISocketSubsystem::GetLocalAdapterAddresses(TArray<TSharedPtr<FInternetAddr>>& OutAddresses)
 {
 	FString HostName;
+
+	UE_LOG(LogSockets, Warning, TEXT("Falling back to generic GetLocalAdapterAddresses implementation. Consider implementing a Platform specific version."));
+
 	// Attempt to get a hostname so that we can look it up in order to get an idea of adapters that we might have 
 	// (or the addresses that are tied to us). This is a fallback implementation for platforms that do not have this implemented.
 	// Platforms are encouraged to implement this themselves.
