Add missing copyright notice

diff --git a/Engine/Source/Runtime/Sockets/Private/Android/SocketsAndroid.cpp b/Engine/Source/Runtime/Sockets/Private/Android/SocketsAndroid.cpp
--- a/Engine/Source/Runtime/Sockets/Private/Android/SocketsAndroid.cpp
+++ b/Engine/Source/Runtime/Sockets/Private/Android/SocketsAndroid.cpp
@@ -1,3 +1,5 @@
+// Copyright Epic Games, Inc. All Rights Reserved.
+
 #include "SocketsAndroid.h"
 
 FSocketAndroid::~FSocketAndroid()
