Fix WebTests Android compile errors.

diff --git a/Engine/Source/Programs/LowLevelTests/FoundationTests.Target.cs b/Engine/Source/Programs/LowLevelTests/FoundationTests.Target.cs
--- a/Engine/Source/Programs/LowLevelTests/FoundationTests.Target.cs
+++ b/Engine/Source/Programs/LowLevelTests/FoundationTests.Target.cs
@@ -1,6 +1,5 @@
 // Copyright Epic Games, Inc. All Rights Reserved.
 
-using System.IO;
 using UnrealBuildTool;
 
 [SupportedPlatforms(UnrealPlatformClass.All)]
@@ -14,16 +13,5 @@ public class FoundationTestsTarget : TestTargetRules
 
 		bBuildWithEditorOnlyData = Target.Platform.IsInGroup(UnrealPlatformGroup.Desktop)
 			&& (Target.Configuration == UnrealTargetConfiguration.Debug || Target.Configuration == UnrealTargetConfiguration.Development);
-
-		if (Target.Platform == UnrealTargetPlatform.Android)
-		{
-			string VersionScriptFile = Path.GetTempPath() + "LLTWorkaroundScrip.ldscript";
-			using (StreamWriter Writer = File.CreateText(VersionScriptFile))
-			{
-				// Workaround for a linker bug when building LowLevelTests for Android
-				Writer.WriteLine("{ local: *; };");
-			}
-			AdditionalLinkerArguments = " -Wl,--version-script=\"" + VersionScriptFile + "\"";
-		}
 	}
 }
diff --git a/Engine/Source/Programs/UnrealBuildTool/Configuration/TestTargetRules.cs b/Engine/Source/Programs/UnrealBuildTool/Configuration/TestTargetRules.cs
--- a/Engine/Source/Programs/UnrealBuildTool/Configuration/TestTargetRules.cs
+++ b/Engine/Source/Programs/UnrealBuildTool/Configuration/TestTargetRules.cs
@@ -341,6 +341,14 @@ namespace UnrealBuildTool
 				GlobalDefinitions.Add("USE_ANDROID_OPENGL=0");
 				GlobalDefinitions.Add("USE_ANDROID_LAUNCH=0");
 				GlobalDefinitions.Add("USE_ANDROID_JNI=0");
+
+				string VersionScriptFile = Path.GetTempPath() + "LLTWorkaroundScrip.ldscript";
+				using (StreamWriter Writer = System.IO.File.CreateText(VersionScriptFile))
+				{
+					// Workaround for a linker bug when building LowLevelTests for Android
+					Writer.WriteLine("{ local: *; };");
+				}
+				AdditionalLinkerArguments = " -Wl,--version-script=\"" + VersionScriptFile + "\"";
 			}
 			else if (Target.Platform == UnrealTargetPlatform.IOS)
 			{
diff --git a/Engine/Source/Runtime/Sockets/Private/Android/SocketsAndroid.cpp b/Engine/Source/Runtime/Sockets/Private/Android/SocketsAndroid.cpp
--- a/Engine/Source/Runtime/Sockets/Private/Android/SocketsAndroid.cpp
+++ b/Engine/Source/Runtime/Sockets/Private/Android/SocketsAndroid.cpp
@@ -49,6 +49,7 @@ bool FSocketAndroid::JoinMulticastGroup(const FInternetAddr& GroupAddress, const
 
 void FSocketAndroid::AcquireMulticastLock()
 {
+#if USE_ANDROID_JNI
     extern bool AndroidThunkCpp_AcquireWifiManagerMulticastLock();
     if (AndroidThunkCpp_AcquireWifiManagerMulticastLock())
     {
@@ -56,6 +57,7 @@ void FSocketAndroid::AcquireMulticastLock()
 		UE_LOG(LogSockets, VeryVerbose, TEXT("WifiManager.MulticastLock succesfully aquired"));
     }
     else
+#endif
     {
         static bool bAlreadyLoggedFailure = false;  
         if (!bAlreadyLoggedFailure) 
@@ -68,6 +70,7 @@ void FSocketAndroid::AcquireMulticastLock()
 
 void FSocketAndroid::ReleaseMulticastLock()
 {
+#if USE_ANDROID_JNI
     if (bIsMulticastLockAcquired)
     {
 	    UE_LOG(LogSockets, VeryVerbose, TEXT("Releasing WifiManager.MulticastLock"));
@@ -75,4 +78,5 @@ void FSocketAndroid::ReleaseMulticastLock()
         AndroidThunkCpp_ReleaseWifiManagerMulticastLock();
         bIsMulticastLockAcquired = false;
     }
+#endif
 }
\ No newline at end of file
