[TSAN] Fix warning by making activity time atomic on socket

#rnx

diff --git a/Engine/Source/Runtime/Sockets/Private/BSDSockets/SocketsBSD.h b/Engine/Source/Runtime/Sockets/Private/BSDSockets/SocketsBSD.h
--- a/Engine/Source/Runtime/Sockets/Private/BSDSockets/SocketsBSD.h
+++ b/Engine/Source/Runtime/Sockets/Private/BSDSockets/SocketsBSD.h
@@ -6,6 +6,8 @@
 #include "HAL/PlatformTime.h"
 #include "BSDSockets/SocketSubsystemBSDPrivate.h"
 
+#include <atomic>
+
 class FInternetAddr;
 
 #if PLATFORM_HAS_BSD_SOCKETS || PLATFORM_HAS_BSD_IPV6_SOCKETS
@@ -43,21 +45,6 @@ class FSocketBSD
 {
 public:
 
-	/**
-	 * Assigns a BSD socket to this object.
-	 *
-	 * @param InSocket the socket to assign to this object.
-	 * @param InSocketType the type of socket that was created.
-	 * @param InSocketDescription the debug description of the socket.
-	 */
-	UE_DEPRECATED(4.22, "Use the socket constructor that specifies protocol stack for better compatibility and debugging")
-	FSocketBSD(SOCKET InSocket, ESocketType InSocketType, const FString& InSocketDescription, ISocketSubsystem * InSubsystem)
-		: FSocket(InSocketType, InSocketDescription, NAME_None)
-		, Socket(InSocket)
-		, LastActivityTime(0)
-		, SocketSubsystem(InSubsystem)
-	{ }
-
 	/**
 	 * Assigns a BSD socket to this object.
 	 *
@@ -152,7 +139,7 @@ protected:
 	SOCKET Socket;
 
 	/** Last activity time. */
-	double LastActivityTime;
+	std::atomic<double> LastActivityTime { 0.0 };
 
 	/** Pointer to the subsystem that created it. */
 	ISocketSubsystem* SocketSubsystem;
