Sony: Sockets: Fixed possible crashes with SIGPIPE when using SendTo on a socket closed on the other end.

The problem was caused by not passing MSG_NOSIGNAL to sendto in some cases. It was already done for send either directly using passing the flag, or by using platform-specific API which handles it internally.

Most of the change is removal of some old code comments with the previous implementation.

#rnx

diff --git a/Engine/Source/Runtime/Sockets/Private/BSDSockets/SocketsBSD.cpp b/Engine/Source/Runtime/Sockets/Private/BSDSockets/SocketsBSD.cpp
--- a/Engine/Source/Runtime/Sockets/Private/BSDSockets/SocketsBSD.cpp
+++ b/Engine/Source/Runtime/Sockets/Private/BSDSockets/SocketsBSD.cpp
@@ -167,7 +167,7 @@ bool FSocketBSD::SendTo(const uint8* Data, int32 Count, int32& BytesSent, const
 
 	const FInternetAddrBSD& BSDAddr = static_cast<const FInternetAddrBSD&>(Destination);
 	// Write the data and see how much was written
-	BytesSent = sendto(Socket, (const char*)Data, Count, 0, (const sockaddr*)&(BSDAddr.Addr), BSDAddr.GetStorageSize());
+	BytesSent = sendto(Socket, (const char*)Data, Count, SendFlags, (const sockaddr*)&(BSDAddr.Addr), BSDAddr.GetStorageSize());
 
 //	NETWORK_PROFILER(FSocket::SendTo(Data,Count,BytesSent,Destination));
 
@@ -181,7 +181,7 @@ bool FSocketBSD::SendTo(const uint8* Data, int32 Count, int32& BytesSent, const
 
 bool FSocketBSD::Send(const uint8* Data, int32 Count, int32& BytesSent)
 {
-	BytesSent = send(Socket,(const char*)Data,Count,0);
+	BytesSent = send(Socket,(const char*)Data,Count,SendFlags);
 
 //	NETWORK_PROFILER(FSocket::Send(Data,Count,BytesSent));
 
diff --git a/Engine/Source/Runtime/Sockets/Private/BSDSockets/SocketsBSD.h b/Engine/Source/Runtime/Sockets/Private/BSDSockets/SocketsBSD.h
--- a/Engine/Source/Runtime/Sockets/Private/BSDSockets/SocketsBSD.h
+++ b/Engine/Source/Runtime/Sockets/Private/BSDSockets/SocketsBSD.h
@@ -141,6 +141,9 @@ protected:
 	/** Last activity time. */
 	std::atomic<double> LastActivityTime { 0.0 };
 
+	/** Flags which will be passed to send and sendto. */
+	int SendFlags = 0;
+
 	/** Pointer to the subsystem that created it. */
 	ISocketSubsystem* SocketSubsystem;
 };
