Clean up platform extensions headers for BSD sockets.

#rnx

diff --git a/Engine/Source/Runtime/Sockets/Private/BSDSockets/SocketSubsystemBSDPrivate.h b/Engine/Source/Runtime/Sockets/Private/BSDSockets/SocketSubsystemBSDPrivate.h
--- a/Engine/Source/Runtime/Sockets/Private/BSDSockets/SocketSubsystemBSDPrivate.h
+++ b/Engine/Source/Runtime/Sockets/Private/BSDSockets/SocketSubsystemBSDPrivate.h
@@ -15,8 +15,8 @@
 
 	#include "Windows/HideWindowsPlatformTypes.h"
 #else
-#if PLATFORM_SWITCH
-	#include "SwitchSocketApiWrapper.h"
+#ifdef UE_SOCKET_API_WRAPPER
+	#include UE_SOCKET_API_WRAPPER
 #else
 	#include <unistd.h>
 	#include <sys/socket.h>
