Fix uninitialized counters in MultichannelTcpSocket causing data not to be sent/received.

diff --git a/Engine/Source/Runtime/Sockets/Public/MultichannelTcpSender.h b/Engine/Source/Runtime/Sockets/Public/MultichannelTcpSender.h
--- a/Engine/Source/Runtime/Sockets/Public/MultichannelTcpSender.h
+++ b/Engine/Source/Runtime/Sockets/Public/MultichannelTcpSender.h
@@ -44,7 +44,8 @@ public:
 	 * @param InOkToSendDelegate Delegate to ask if it is ok to send a packet.
 	 */
 	FMultichannelTcpSender( FSocket* InSocket, const FOnMultichannelTcpOkToSend& InOkToSendDelegate )
-		: Socket(InSocket)
+		: BytesSent(0)
+		, Socket(InSocket)
 		, OkToSendDelegate(InOkToSendDelegate)
 	{
 		Thread = FRunnableThread::Create(this, TEXT("FMultichannelTCPSender"), 8 * 1024, TPri_AboveNormal);
diff --git a/Engine/Source/Runtime/Sockets/Public/MultichannelTcpSocket.h b/Engine/Source/Runtime/Sockets/Public/MultichannelTcpSocket.h
--- a/Engine/Source/Runtime/Sockets/Public/MultichannelTcpSocket.h
+++ b/Engine/Source/Runtime/Sockets/Public/MultichannelTcpSocket.h
@@ -41,6 +41,7 @@ public:
 	 */
 	FMultichannelTcpSocket( FSocket* InSocket, uint64 InBandwidthLatencyProduct )
 		: BandwidthLatencyProduct(InBandwidthLatencyProduct)
+		, RemoteReceiverBytesReceived(0)
 		, Receiver(InSocket, FOnMultichannelTcpReceive::CreateRaw(this, &FMultichannelTcpSocket::HandleReceiverReceive))
 		, Sender(InSocket, FOnMultichannelTcpOkToSend::CreateRaw(this, &FMultichannelTcpSocket::HandleSenderOkToSend))
 		, Socket(InSocket)
@@ -282,12 +283,12 @@ private:
 	/** Holds a critical section to guard the receive buffers. */
 	FCriticalSection ReceiveBuffersCriticalSection;
 
-	/** Holds the receiver thread. */
-	FMultichannelTcpReceiver Receiver;
-
 	/** Holds the total number of bytes received by the client (comes from an 'Ack' on the control channel). */
 	int64 RemoteReceiverBytesReceived;
 
+	/** Holds the receiver thread. */
+	FMultichannelTcpReceiver Receiver;
+
 	/** Holds the sender thread. */
 	FMultichannelTcpSender Sender;
 
