Open the FInternetAddrBSD::ToString function

diff --git a/Engine/Source/Runtime/Sockets/Private/BSDSockets/IPAddressBSD.cpp b/Engine/Source/Runtime/Sockets/Private/BSDSockets/IPAddressBSD.cpp
--- a/Engine/Source/Runtime/Sockets/Private/BSDSockets/IPAddressBSD.cpp
+++ b/Engine/Source/Runtime/Sockets/Private/BSDSockets/IPAddressBSD.cpp
@@ -4,6 +4,7 @@
 
 #include "Math/NumericLimits.h"
 #include "SocketSubsystemBSD.h"
+#include "AutoRTFM/AutoRTFM.h"
 
 #if PLATFORM_HAS_BSD_SOCKETS || PLATFORM_HAS_BSD_IPV6_SOCKETS
 
@@ -457,39 +458,42 @@ void FInternetAddrBSD::SetIPv6LoopbackAddress()
 
 FString FInternetAddrBSD::ToString(bool bAppendPort) const
 {
-	FString ReturnVal("");
-#if PLATFORM_HAS_BSD_SOCKET_FEATURE_GETNAMEINFO
-	char IPStr[NI_MAXHOST];
-	if (getnameinfo((const sockaddr*)&Addr, GetStorageSize(), IPStr, NI_MAXHOST, nullptr, 0, NI_NUMERICHOST) == 0)
+	return AutoRTFM::Open([&]
 	{
-		if (GetProtocolType() == FNetworkProtocolTypes::IPv6)
+		FString ReturnVal("");
+#if PLATFORM_HAS_BSD_SOCKET_FEATURE_GETNAMEINFO
+		char IPStr[NI_MAXHOST];
+		if (getnameinfo((const sockaddr*)&Addr, GetStorageSize(), IPStr, NI_MAXHOST, nullptr, 0, NI_NUMERICHOST) == 0)
 		{
-			FString IPv6Str(ANSI_TO_TCHAR(IPStr));
-			// Remove the scope interface if it exists.
-			const int32 InterfaceMarkerIndex = IPv6Str.Find("%", ESearchCase::CaseSensitive, ESearchDir::FromEnd);
-			if (InterfaceMarkerIndex != INDEX_NONE)
+			if (GetProtocolType() == FNetworkProtocolTypes::IPv6)
 			{
-				IPv6Str.LeftInline(InterfaceMarkerIndex, EAllowShrinking::No);
+				FString IPv6Str(ANSI_TO_TCHAR(IPStr));
+				// Remove the scope interface if it exists.
+				const int32 InterfaceMarkerIndex = IPv6Str.Find("%", ESearchCase::CaseSensitive, ESearchDir::FromEnd);
+				if (InterfaceMarkerIndex != INDEX_NONE)
+				{
+					IPv6Str.LeftInline(InterfaceMarkerIndex, EAllowShrinking::No);
+				}
+
+				// Using dynamic formatting strings are deprecated.
+				ReturnVal = bAppendPort ? FString::Printf(TEXT("[%s]"), *IPv6Str) 
+					: FString::Printf(TEXT("%s"), *IPv6Str);
+			}
+			else
+			{
+				ReturnVal = IPStr;
 			}
 
-			// Using dynamic formatting strings are deprecated.
-			ReturnVal = bAppendPort ? FString::Printf(TEXT("[%s]"), *IPv6Str) 
-				: FString::Printf(TEXT("%s"), *IPv6Str);
-		}
-		else
-		{
-			ReturnVal = IPStr;
-		}
-
-		if (bAppendPort)
-		{
-			ReturnVal += FString::Printf(TEXT(":%d"), GetPort());
+			if (bAppendPort)
+			{
+				ReturnVal += FString::Printf(TEXT(":%d"), GetPort());
+			}
 		}
-	}
 #else
-	UE_LOG(LogSockets, Error, TEXT("Platform has no getnameinfo(), but did not override FInternetAddrBSD::ToString()"));
+		UE_LOG(LogSockets, Error, TEXT("Platform has no getnameinfo(), but did not override FInternetAddrBSD::ToString()"));
 #endif
-	return ReturnVal;
+		return ReturnVal;
+	});
 }
 
 bool FInternetAddrBSD::operator==(const FInternetAddr& Other) const
