Mac: Fix disconnecting bug with sockets.

diff --git a/Engine/Source/Runtime/Sockets/Private/Mac/SocketSubsystemMac.h b/Engine/Source/Runtime/Sockets/Private/Mac/SocketSubsystemMac.h
--- a/Engine/Source/Runtime/Sockets/Private/Mac/SocketSubsystemMac.h
+++ b/Engine/Source/Runtime/Sockets/Private/Mac/SocketSubsystemMac.h
@@ -95,11 +95,6 @@ public:
 		FSocketMac::Close();
 	}
 
-	virtual bool Close() override
-	{
-		return close(Socket) == 0;
-	}
-
 	virtual bool SetReuseAddr(bool bAllowReuse) override
 	{
 		int Param = bAllowReuse ? 1 : 0;
