Remove failure for Android in FSocketSubsystemBSD::CreateSocket when socket file descriptor >= 1024

diff --git a/Engine/Source/Runtime/Sockets/Private/BSDSockets/SocketSubsystemBSD.cpp b/Engine/Source/Runtime/Sockets/Private/BSDSockets/SocketSubsystemBSD.cpp
--- a/Engine/Source/Runtime/Sockets/Private/BSDSockets/SocketSubsystemBSD.cpp
+++ b/Engine/Source/Runtime/Sockets/Private/BSDSockets/SocketSubsystemBSD.cpp
@@ -51,17 +51,7 @@ FSocket* FSocketSubsystemBSD::CreateSocket(const FName& SocketType, const FStrin
 
 	Socket = socket(GetProtocolFamilyValue(SocketProtocolType), SocketTypeFlag | PlatformSpecificTypeFlags, ((bIsUDP) ? IPPROTO_UDP : IPPROTO_TCP));
 
-#if PLATFORM_ANDROID
-	// To avoid out of range in FD_SET
-	if (Socket != INVALID_SOCKET && Socket >= 1024)
-	{
-		closesocket(Socket);
-	}
-	else
-#endif
-	{
-		NewSocket = (Socket != INVALID_SOCKET) ? InternalBSDSocketFactory(Socket, ((bIsUDP) ? SOCKTYPE_Datagram : SOCKTYPE_Streaming), SocketDescription, SocketProtocolType) : nullptr;
-	}
+	NewSocket = (Socket != INVALID_SOCKET) ? InternalBSDSocketFactory(Socket, ((bIsUDP) ? SOCKTYPE_Datagram : SOCKTYPE_Streaming), SocketDescription, SocketProtocolType) : nullptr;
 
 	if (!NewSocket)
 	{
