Extracting FQDN via getnameinfo if getaddrinfo only returns a hostname
#rnx

diff --git a/Engine/Source/Runtime/Sockets/Private/BSDSockets/SocketSubsystemBSD.cpp b/Engine/Source/Runtime/Sockets/Private/BSDSockets/SocketSubsystemBSD.cpp
--- a/Engine/Source/Runtime/Sockets/Private/BSDSockets/SocketSubsystemBSD.cpp
+++ b/Engine/Source/Runtime/Sockets/Private/BSDSockets/SocketSubsystemBSD.cpp
@@ -89,6 +89,12 @@ FAddressInfoResult FSocketSubsystemBSD::GetAddressInfo(const TCHAR* HostName, co
 		QueryFlags |= EAddressInfoFlags::NoResolveService;
 	}
 
+#if PLATFORM_HAS_BSD_SOCKET_FEATURE_GETNAMEINFO
+	// On Windows getaddrinfo will likely return a hostname instead of FQDN in ai_canonname.
+	// Try getnameinfo on every associated ip address and choose a first name with a dot if ai_canonname didn't contain FQDN.
+	bool bTryGetFQDN = EnumHasAnyFlags(QueryFlags, EAddressInfoFlags::FQDomainName);
+#endif
+
 	addrinfo HintAddrInfo;
 	FMemory::Memzero(&HintAddrInfo, sizeof(HintAddrInfo));
 	HintAddrInfo.ai_family = GetProtocolFamilyValue(ProtocolTypeName);
@@ -112,6 +118,16 @@ FAddressInfoResult FSocketSubsystemBSD::GetAddressInfo(const TCHAR* HostName, co
 		if (AddrInfo != nullptr && AddrInfo->ai_canonname != nullptr)
 		{
 			AddrQueryResult.CanonicalNameResult = UTF8_TO_TCHAR(AddrInfo->ai_canonname);
+			UE_LOG(LogSockets, Verbose, TEXT("getaddrinfo ai_canonname returns: %s"), *AddrQueryResult.CanonicalNameResult);
+
+#if PLATFORM_HAS_BSD_SOCKET_FEATURE_GETNAMEINFO
+			// ai_canonname contained a dot, likely to be the best FQDN
+			int32 DotPosition;
+			if (AddrQueryResult.CanonicalNameResult.FindChar(TEXT('.'), DotPosition))
+			{
+				bTryGetFQDN = false;
+			}
+#endif
 		}
 
 		for (; AddrInfo != nullptr; AddrInfo = AddrInfo->ai_next)
@@ -135,6 +151,30 @@ FAddressInfoResult FSocketSubsystemBSD::GetAddressInfo(const TCHAR* HostName, co
 					{
 						sockaddr_in6* IPv6AddrData = reinterpret_cast<sockaddr_in6*>(AddrInfo->ai_addr);
 						UE_LOG(LogSockets, Verbose, TEXT("# Family: %s Address: %s And scope %d"), ((AddrInfo->ai_family == AF_INET) ? TEXT("IPv4") : TEXT("IPv6")), *(NewAddress->ToString(true)), ((AddrInfo->ai_family == AF_INET) ? -1 : IPv6AddrData->sin6_scope_id));
+
+#if PLATFORM_HAS_BSD_SOCKET_FEATURE_GETNAMEINFO
+						// Try to acquire the FQDN
+						if (bTryGetFQDN && (AddrData->ss_family == AF_INET || AddrData->ss_family == AF_INET6))
+						{
+							char HostBuffer[NI_MAXHOST];
+							if (getnameinfo(
+								reinterpret_cast<const sockaddr*>(AddrData), AddrInfo->ai_addrlen,
+								HostBuffer, sizeof(HostBuffer),
+								nullptr, 0,
+								NI_NAMEREQD) == 0)
+							{
+								// Use HostName as FQDN if it contains a dot
+								int32 DotPosition;
+								if (FAnsiStringView(HostBuffer).FindChar(TEXT('.'), DotPosition))
+								{
+									AddrQueryResult.CanonicalNameResult = StringCast<TCHAR>(HostBuffer);
+									bTryGetFQDN = false;
+
+									UE_LOG(LogSockets, Verbose, TEXT("Using getnameinfo hostname '%s'"), *AddrQueryResult.CanonicalNameResult);
+								}
+							}
+						}
+#endif
 					}
 				}
 			}
