- Add SocketSubsystemModule.cpp to match SocketSubsystemModule.h and move the module bits out of SocketSubsystem.cpp

diff --git a/Engine/Source/Runtime/Sockets/Private/SocketSubsystem.cpp b/Engine/Source/Runtime/Sockets/Private/SocketSubsystem.cpp
--- a/Engine/Source/Runtime/Sockets/Private/SocketSubsystem.cpp
+++ b/Engine/Source/Runtime/Sockets/Private/SocketSubsystem.cpp
@@ -14,8 +14,6 @@
 
 DEFINE_LOG_CATEGORY(LogSockets);
 
-IMPLEMENT_MODULE( FSocketSubsystemModule, Sockets );
-
 /** Values for the Protocol Typenames */
 namespace FNetworkProtocolTypes
 {
@@ -23,60 +21,6 @@ namespace FNetworkProtocolTypes
 	const FLazyName IPv6(TEXT("IPv6"));
 }
 
-/** Each platform will implement these functions to construct/destroy socket implementations */
-extern FName CreateSocketSubsystem(FSocketSubsystemModule& SocketSubsystemModule);
-extern void DestroySocketSubsystem(FSocketSubsystemModule& SocketSubsystemModule);
-
-/** Helper function to turn the friendly subsystem name into the module name */
-static inline FName GetSocketModuleName(const FString& SubsystemName)
-{
-	FName ModuleName;
-	FString SocketBaseName("Sockets");
-
-	if (!SubsystemName.StartsWith(SocketBaseName, ESearchCase::CaseSensitive))
-	{
-		ModuleName = FName(*(SocketBaseName + SubsystemName));
-	}
-	else
-	{
-		ModuleName = FName(*SubsystemName);
-	}
-
-	return ModuleName;
-}
-
-/**
- * Helper function that loads a given platform service module if it isn't already loaded
- *
- * @param SubsystemName Name of the requested platform service to load
- * @return The module interface of the requested platform service, NULL if the service doesn't exist
- */
-static IModuleInterface* LoadSubsystemModule(const FString& SubsystemName)
-{
-#if !UE_BUILD_SHIPPING && !UE_BUILD_SHIPPING_WITH_EDITOR
-	// Early out if we are overriding the module load
-	bool bAttemptLoadModule = !FParse::Param(FCommandLine::Get(), *FString::Printf(TEXT("no%s"), *SubsystemName));
-	if (bAttemptLoadModule)
-#endif
-	{
-		FName ModuleName;
-		FModuleManager& ModuleManager = FModuleManager::Get();
-
-		ModuleName = GetSocketModuleName(SubsystemName);
-		if (!ModuleManager.IsModuleLoaded(ModuleName))
-		{
-			// Attempt to load the module
-			ModuleManager.LoadModule(ModuleName);
-		}
-
-		return ModuleManager.GetModule(ModuleName);
-	}
-
-#if !UE_BUILD_SHIPPING && !UE_BUILD_SHIPPING_WITH_EDITOR
-	return nullptr;
-#endif
-}
-
 ISocketSubsystem::ISocketSubsystem() = default;
 ISocketSubsystem::~ISocketSubsystem() = default;
 
@@ -95,7 +39,7 @@ FUniqueSocket ISocketSubsystem::CreateUniqueSocket(const FName& SocketType, cons
  */
 void ISocketSubsystem::ShutdownAllSystems()
 {
-	if (IsInGameThread() && FModuleManager::Get().IsModuleLoaded(TEXT("Sockets")) == true)
+	if (IsInGameThread() && FModuleManager::Get().IsModuleLoaded(TEXT("Sockets")))
 	{
 		// Unloading the Sockets module will call FSocketSubsystemModule::ShutdownSocketSubsystem()
 		const bool bIsShutdown = true;
@@ -103,124 +47,10 @@ void ISocketSubsystem::ShutdownAllSystems()
 	}
 }
 
-
-/**
- * Called right after the module DLL has been loaded and the module object has been created
- * Overloaded to allow the default subsystem a chance to load
- */
-void FSocketSubsystemModule::StartupModule()
-{
-	FString InterfaceString;
-
-	// Initialize the platform defined socket subsystem first
-	DefaultSocketSubsystem = CreateSocketSubsystem( *this );
-}
-
-/**
- * Called before the module is unloaded, right before the module object is destroyed.
- * Overloaded to shut down all loaded online subsystems
- */
-void FSocketSubsystemModule::ShutdownModule()
-{
-	ShutdownSocketSubsystem();
-}
-
-void FSocketSubsystemModule::ShutdownSocketSubsystem()
-{
-	// Destroy the platform defined socket subsystem first
-	DestroySocketSubsystem( *this );
-
-	FModuleManager& ModuleManager = FModuleManager::Get();
-	// Unload all the supporting factories
-	for (TMap<FName, ISocketSubsystem*>::TIterator It(SocketSubsystems); It; ++It)
-	{
-		It.Value()->Shutdown();
-		// Unloading the module will do proper cleanup
-		FName ModuleName = GetSocketModuleName(It.Key().ToString());
-
-		const bool bIsShutdown = true;
-		ModuleManager.UnloadModule(ModuleName, bIsShutdown);
-	} 
-	//ensure(SocketSubsystems.Num() == 0);
-}
-
-
-/** 
- * Register a new socket subsystem interface with the base level factory provider
- *
- * @param FactoryName - name of subsystem as referenced by consumers
- * @param Factory - instantiation of the socket subsystem interface, this will take ownership
- * @param bMakeDefault - make this subsystem the default
- */
-void FSocketSubsystemModule::RegisterSocketSubsystem(const FName FactoryName, ISocketSubsystem* Factory, bool bMakeDefault)
-{
-	if (!SocketSubsystems.Contains(FactoryName))
-	{
-		SocketSubsystems.Add(FactoryName, Factory);
-	}
-
-	if (bMakeDefault)
-	{
-		DefaultSocketSubsystem = FactoryName;
-	}
-}
-
-/** 
- * Unregister an existing online subsystem interface from the base level factory provider
- * @param FactoryName - name of subsystem as referenced by consumers
- */
-void FSocketSubsystemModule::UnregisterSocketSubsystem(const FName FactoryName)
-{
-	if (SocketSubsystems.Contains(FactoryName))
-	{
-		SocketSubsystems.Remove(FactoryName);
-	}
-}
-
-
-/** 
- * Main entry point for accessing a socket subsystem by name
- * Will load the appropriate module if the subsystem isn't currently loaded
- * It's possible that the subsystem doesn't exist and therefore can return NULL
- *
- * @param SubsystemName - name of subsystem as referenced by consumers
- * @return Requested socket subsystem, or NULL if that subsystem was unable to load or doesn't exist
- */
-ISocketSubsystem* FSocketSubsystemModule::GetSocketSubsystem(const FName InSubsystemName)
-{
-	FName SubsystemName = InSubsystemName;
-	if (SubsystemName == NAME_None)
-	{
-		SubsystemName = DefaultSocketSubsystem;
-	}
-
-	ISocketSubsystem** SocketSubsystemFactory = SocketSubsystems.Find(SubsystemName);
-	if (SocketSubsystemFactory == nullptr)
-	{
-		// Attempt to load the requested factory
-		IModuleInterface* NewModule = LoadSubsystemModule(SubsystemName.ToString());
-		if (NewModule)
-		{
-			// If the module loaded successfully this should be non-NULL;
-			SocketSubsystemFactory = SocketSubsystems.Find(SubsystemName);
-		}
-
-		if (SocketSubsystemFactory == nullptr)
-		{
-			UE_LOG(LogSockets, Warning, TEXT("Unable to load SocketSubsystem module %s"), *InSubsystemName.ToString());
-		}
-	}
-
-	return (SocketSubsystemFactory == nullptr) ? nullptr : *SocketSubsystemFactory;
-}
-
-
-
 //////////////////////////////////
 // ISocketSubsystem
 /////////////////////////////////
 
-
 ISocketSubsystem* ISocketSubsystem::Get(const FName& SubsystemName)
 {
 	// wrap the platform file with a logger
diff --git a/Engine/Source/Runtime/Sockets/Private/SocketSubsystemModule.cpp b/Engine/Source/Runtime/Sockets/Private/SocketSubsystemModule.cpp
new file mode 100644
--- /dev/null
+++ b/Engine/Source/Runtime/Sockets/Private/SocketSubsystemModule.cpp
@@ -0,0 +1,138 @@
+// Copyright Epic Games, Inc. All Rights Reserved.
+
+#include "SocketSubsystemModule.h"
+
+#include "Misc/CommandLine.h"
+#include "Modules/ModuleManager.h"
+#include "SocketSubsystem.h"
+
+IMPLEMENT_MODULE( FSocketSubsystemModule, Sockets );
+
+/** Each platform will implement these functions to construct/destroy socket implementations */
+extern FName CreateSocketSubsystem(FSocketSubsystemModule& SocketSubsystemModule);
+extern void DestroySocketSubsystem(FSocketSubsystemModule& SocketSubsystemModule);
+
+/** Helper function to turn the friendly subsystem name into the module name */
+static inline FName GetSocketModuleName(const FString& SubsystemName)
+{
+	FName ModuleName;
+	FString SocketBaseName("Sockets");
+
+	if (!SubsystemName.StartsWith(SocketBaseName, ESearchCase::CaseSensitive))
+	{
+		ModuleName = FName(*(SocketBaseName + SubsystemName));
+	}
+	else
+	{
+		ModuleName = FName(*SubsystemName);
+	}
+
+	return ModuleName;
+}
+
+/**
+ * Helper function that loads a given platform service module if it isn't already loaded
+ *
+ * @param SubsystemName Name of the requested platform service to load
+ * @return The module interface of the requested platform service, NULL if the service doesn't exist
+ */
+static IModuleInterface* LoadSubsystemModule(const FString& SubsystemName)
+{
+#if !UE_BUILD_SHIPPING && !UE_BUILD_SHIPPING_WITH_EDITOR
+	// Early out if we are overriding the module load
+	bool bAttemptLoadModule = !FParse::Param(FCommandLine::Get(), *FString::Printf(TEXT("no%s"), *SubsystemName));
+	if (bAttemptLoadModule)
+#endif
+	{
+		FName ModuleName;
+		FModuleManager& ModuleManager = FModuleManager::Get();
+
+		ModuleName = GetSocketModuleName(SubsystemName);
+		if (!ModuleManager.IsModuleLoaded(ModuleName))
+		{
+			// Attempt to load the module
+			ModuleManager.LoadModule(ModuleName);
+		}
+
+		return ModuleManager.GetModule(ModuleName);
+	}
+
+#if !UE_BUILD_SHIPPING && !UE_BUILD_SHIPPING_WITH_EDITOR
+	return nullptr;
+#endif
+}
+
+ISocketSubsystem* FSocketSubsystemModule::GetSocketSubsystem(const FName InSubsystemName)
+{
+	FName SubsystemName = InSubsystemName;
+	if (SubsystemName == NAME_None)
+	{
+		SubsystemName = DefaultSocketSubsystem;
+	}
+
+	ISocketSubsystem** SocketSubsystemFactory = SocketSubsystems.Find(SubsystemName);
+	if (SocketSubsystemFactory == nullptr)
+	{
+		// Attempt to load the requested factory
+		IModuleInterface* NewModule = LoadSubsystemModule(SubsystemName.ToString());
+		if (NewModule)
+		{
+			// If the module loaded successfully this should be non-NULL;
+			SocketSubsystemFactory = SocketSubsystems.Find(SubsystemName);
+		}
+
+		if (SocketSubsystemFactory == nullptr)
+		{
+			UE_LOG(LogSockets, Warning, TEXT("Unable to load SocketSubsystem module %s"), *InSubsystemName.ToString());
+		}
+	}
+
+	return (SocketSubsystemFactory == nullptr) ? nullptr : *SocketSubsystemFactory;
+}
+
+void FSocketSubsystemModule::RegisterSocketSubsystem(const FName FactoryName, ISocketSubsystem* Factory, bool bMakeDefault)
+{
+	if (!SocketSubsystems.Contains(FactoryName))
+	{
+		SocketSubsystems.Add(FactoryName, Factory);
+	}
+
+	if (bMakeDefault)
+	{
+		DefaultSocketSubsystem = FactoryName;
+	}
+}
+
+void FSocketSubsystemModule::UnregisterSocketSubsystem(const FName FactoryName)
+{
+	if (SocketSubsystems.Contains(FactoryName))
+	{
+		SocketSubsystems.Remove(FactoryName);
+	}
+}
+
+void FSocketSubsystemModule::StartupModule()
+{
+	FString InterfaceString;
+
+	// Initialize the platform defined socket subsystem first
+	DefaultSocketSubsystem = CreateSocketSubsystem( *this );
+}
+
+void FSocketSubsystemModule::ShutdownModule()
+{
+	// Destroy the platform defined socket subsystem first
+	DestroySocketSubsystem( *this );
+
+	FModuleManager& ModuleManager = FModuleManager::Get();
+	// Unload all the supporting factories
+	for (TMap<FName, ISocketSubsystem*>::TIterator It(SocketSubsystems); It; ++It)
+	{
+		It.Value()->Shutdown();
+		// Unloading the module will do proper cleanup
+		FName ModuleName = GetSocketModuleName(It.Key().ToString());
+
+		const bool bIsShutdown = true;
+		ModuleManager.UnloadModule(ModuleName, bIsShutdown);
+	}
+}
diff --git a/Engine/Source/Runtime/Sockets/Public/SocketSubsystemModule.h b/Engine/Source/Runtime/Sockets/Public/SocketSubsystemModule.h
--- a/Engine/Source/Runtime/Sockets/Public/SocketSubsystemModule.h
+++ b/Engine/Source/Runtime/Sockets/Public/SocketSubsystemModule.h
@@ -19,11 +19,6 @@ private:
 	/** Mapping of all currently loaded subsystems to their name */
 	TMap<FName, class ISocketSubsystem*> SocketSubsystems;
 
-	/**
-	 *	Shuts down all registered socket subsystem and unloads their modules
-	 */
-	virtual void ShutdownSocketSubsystem();
-
 public:
 
 	FSocketSubsystemModule() {}
